[TmvMedia] Fix compile errors. CIS Issue 1135911

#Synced-CL 47720872

diff --git a/Engine/Plugins/Media/TmvMedia/Source/TmvMedia/Private/SampleConverter/TmvMediaFrameMipStructuredBuffer.cpp b/Engine/Plugins/Media/TmvMedia/Source/TmvMedia/Private/SampleConverter/TmvMediaFrameMipStructuredBuffer.cpp
--- a/Engine/Plugins/Media/TmvMedia/Source/TmvMedia/Private/SampleConverter/TmvMediaFrameMipStructuredBuffer.cpp
+++ b/Engine/Plugins/Media/TmvMedia/Source/TmvMedia/Private/SampleConverter/TmvMediaFrameMipStructuredBuffer.cpp
@@ -190,7 +190,7 @@ void FTmvMediaFrameMipStructuredBuffer::CopyTileRegions(int32 InFrameId, const F
 	ENQUEUE_RENDER_COMMAND(CopyFromUploadBuffer)([BufferData, BufferRegionsToCopy = MoveTemp(TmpBufferRegionsToCopy), FrameId = InFrameId](FRHICommandListImmediate& InRHICmdList)
 	{
 		RHI_BREADCRUMB_EVENT_STAT_F(InRHICmdList, TmvConverter_UploadBufferRegions, "TmvMediaConverter.StartCopy", "TmvMediaConverter.StartCopy %d", FrameId);
-		SCOPED_GPU_STAT(RHICmdList, TmvConverter_UploadBufferRegions);
+		SCOPED_GPU_STAT(InRHICmdList, TmvConverter_UploadBufferRegions);
 
 		FTmvMediaFrameMipStructuredBuffer* Buffer = static_cast<FTmvMediaFrameMipStructuredBuffer*>(BufferData.Get());
 
diff --git a/Engine/Plugins/Media/TmvMedia/Source/TmvMedia/Private/SampleConverter/TmvMediaFrameMipTextureBuffer.cpp b/Engine/Plugins/Media/TmvMedia/Source/TmvMedia/Private/SampleConverter/TmvMediaFrameMipTextureBuffer.cpp
--- a/Engine/Plugins/Media/TmvMedia/Source/TmvMedia/Private/SampleConverter/TmvMediaFrameMipTextureBuffer.cpp
+++ b/Engine/Plugins/Media/TmvMedia/Source/TmvMedia/Private/SampleConverter/TmvMediaFrameMipTextureBuffer.cpp
@@ -428,7 +428,7 @@ void FTmvMediaFrameMipTextureBuffer::CopyTileRegions(int32 InFrameId, const FTmv
 	ENQUEUE_RENDER_COMMAND(CopyFromUploadBuffer)([BufferData, BufferRegionsToCopy = MoveTemp(TmpBufferRegionsToCopy), FrameId = InFrameId](FRHICommandListImmediate& InRHICmdList)
 	{
 		RHI_BREADCRUMB_EVENT_STAT_F(InRHICmdList, TmvConverter_UploadTextureRegions, "TmvMediaConverter.StartCopy", "TmvMediaConverter.StartCopy %d", FrameId);
-		SCOPED_GPU_STAT(RHICmdList, TmvConverter_UploadTextureRegions);
+		SCOPED_GPU_STAT(InRHICmdList, TmvConverter_UploadTextureRegions);
 
 		FTmvMediaFrameMipTextureBuffer* TextureBuffer = static_cast<FTmvMediaFrameMipTextureBuffer*>(BufferData.Get());
 
