[TmvMedia] Fix CIS Issue 1137520 - ApvMedia: Exclude for Arm64 architecture.

We haven't compile OpenApv for that architecture yet.

#Synced-CL 47720872

diff --git a/Engine/Plugins/Media/TmvMedia/TmvMedia.uplugin b/Engine/Plugins/Media/TmvMedia/TmvMedia.uplugin
--- a/Engine/Plugins/Media/TmvMedia/TmvMedia.uplugin
+++ b/Engine/Plugins/Media/TmvMedia/TmvMedia.uplugin
@@ -30,7 +30,8 @@
 			"Name" : "ApvMedia",
 			"Type" : "Runtime",
 			"LoadingPhase" : "PostEngineInit",
-			"PlatformAllowList" : [ "Win64" ]
+			"PlatformAllowList" : [ "Win64" ],
+			"PlatformArchitectureDenyList": [ "Win64:arm64" ]
 		}
 	]
 }
