[TmvMedia] TMV Transcoding Framework

#Synced-CL 49891360

diff --git a/Engine/Plugins/Media/ImgMedia/Source/OpenExrWrapper/OpenExrWrapper.Build.cs b/Engine/Plugins/Media/ImgMedia/Source/OpenExrWrapper/OpenExrWrapper.Build.cs
--- a/Engine/Plugins/Media/ImgMedia/Source/OpenExrWrapper/OpenExrWrapper.Build.cs
+++ b/Engine/Plugins/Media/ImgMedia/Source/OpenExrWrapper/OpenExrWrapper.Build.cs
@@ -13,7 +13,9 @@ namespace UnrealBuildTool.Rules
             PublicDependencyModuleNames.AddRange(
 				new string[] {
 					"Core",
+					"CoreUObject",
 					"TimeManagement",
+					"TmvMedia",
 				});
 
 			bool bLinuxEnabled = Target.Platform == UnrealTargetPlatform.Linux && Target.Architecture == UnrealArch.X64;
diff --git a/Engine/Plugins/Media/ImgMedia/Source/OpenExrWrapper/Private/OpenExrWrapper.cpp b/Engine/Plugins/Media/ImgMedia/Source/OpenExrWrapper/Private/OpenExrWrapper.cpp
--- a/Engine/Plugins/Media/ImgMedia/Source/OpenExrWrapper/Private/OpenExrWrapper.cpp
+++ b/Engine/Plugins/Media/ImgMedia/Source/OpenExrWrapper/Private/OpenExrWrapper.cpp
@@ -4,9 +4,8 @@
 
 #include <exception>
 #include "Containers/UnrealString.h"
-#include "Logging/LogMacros.h"
 #include "Math/IntRect.h"
-#include "Modules/ModuleManager.h"
+#include "OpenExrWrapperLog.h"
 
 PRAGMA_DEFAULT_VISIBILITY_START
 THIRD_PARTY_INCLUDES_START
@@ -25,10 +24,6 @@ THIRD_PARTY_INCLUDES_START
 THIRD_PARTY_INCLUDES_END
 PRAGMA_DEFAULT_VISIBILITY_END
 
-DECLARE_LOG_CATEGORY_EXTERN(LogOpenEXRWrapper, Log, All);
-DEFINE_LOG_CATEGORY(LogOpenEXRWrapper);
-
-
 /* FOpenExr
  *****************************************************************************/
 
@@ -698,6 +693,4 @@ int32 FTiledOutputFile::GetNumYTiles(int32 MipLevel)
 			TEXT("GetNumYTiles failed: CreateOutputFile has not been called yet."));
 		return 0;
 	}
-}
-
-IMPLEMENT_MODULE(FDefaultModuleImpl, OpenExrWrapper);
+}
\ No newline at end of file
diff --git a/Engine/Plugins/Media/ImgMedia/Source/OpenExrWrapper/Private/OpenExrWrapperTmvEncoder.h b/Engine/Plugins/Media/ImgMedia/Source/OpenExrWrapper/Private/OpenExrWrapperTmvEncoder.h
new file mode 100644
--- /dev/null
+++ b/Engine/Plugins/Media/ImgMedia/Source/OpenExrWrapper/Private/OpenExrWrapperTmvEncoder.h
@@ -0,0 +1,25 @@
+// Copyright Epic Games, Inc. All Rights Reserved.
+
+#pragma once
+
+#include "OpenExrWrapperTmvEncoderOptions.h"
+#include "Encoder/ITmvMediaEncoder.h"
+
+/**
+ * Implementation of a TMV encoder for the OpenExr.
+ */
+class FOpenExrWrapperTmvEncoder : public ITmvMediaEncoder
+{
+public:
+	FOpenExrWrapperTmvEncoder(const FOpenExrWrapperTmvEncoderOptions& InEncoderOptions);
+	virtual ~FOpenExrWrapperTmvEncoder() override;
+
+	//~ Begin ITmvMediaEncoder
+	virtual void RequestMipInfos(const FTmvMediaFrameTimeInfo& InTimeInfo, const FTmvMediaEncoderMipInfo& InFrameInfo, TArray<FTmvMediaFrameMipInfo>& OutFrameMipInfo) override;
+	virtual ETmvMediaEncoderResult Encode(const FTmvMediaFrameTimeInfo& InTimeInfo, ITmvMediaEncoderAccessUnit& InAccessUnit, TArrayView<FTmvMediaEncoderMipRequest> InMipRequests) override;
+	//~ End ITmvMediaEncoder
+
+private:
+	/** Current encoder options */
+	FOpenExrWrapperTmvEncoderOptions EncoderOptions;
+};
\ No newline at end of file
diff --git a/Engine/Plugins/Media/TmvMedia/Source/ApvMedia/ApvMedia.Build.cs b/Engine/Plugins/Media/TmvMedia/Source/ApvMedia/ApvMedia.Build.cs
--- a/Engine/Plugins/Media/TmvMedia/Source/ApvMedia/ApvMedia.Build.cs
+++ b/Engine/Plugins/Media/TmvMedia/Source/ApvMedia/ApvMedia.Build.cs
@@ -16,6 +16,7 @@ namespace UnrealBuildTool.Rules
 				new string[]
 				{
 					"DeveloperSettings",
+					"Engine",
 					"Media",
 					"Projects", 
 					"RHI",
diff --git a/Engine/Plugins/Media/TmvMedia/Source/ApvMedia/Private/Apv/ApvDecoder.cpp b/Engine/Plugins/Media/TmvMedia/Source/ApvMedia/Private/Apv/ApvDecoder.cpp
--- a/Engine/Plugins/Media/TmvMedia/Source/ApvMedia/Private/Apv/ApvDecoder.cpp
+++ b/Engine/Plugins/Media/TmvMedia/Source/ApvMedia/Private/Apv/ApvDecoder.cpp
@@ -2,6 +2,7 @@
 
 #include "ApvDecoder.h"
 #include "ApvMediaLog.h"
+#include "Templates/UniquePtr.h"
 
 namespace UE::ApvMedia
 {
@@ -157,16 +158,16 @@ namespace UE::ApvMedia
 		switch (verbosity)
 		{
 		case OAPV_LOG_ERROR:
-			UE_LOG(LogApvMedia, Error, TEXT("OAPV: %s"), StringCast<TCHAR>(message).Get());
+			UE_LOG(LogApvMedia, Error, TEXT("OpenApv: %s"), StringCast<TCHAR>(message).Get());
 			break;
 		case OAPV_LOG_WARNING:
-			UE_LOG(LogApvMedia, Warning, TEXT("OAPV: %s"), StringCast<TCHAR>(message).Get());
+			UE_LOG(LogApvMedia, Warning, TEXT("OpenApv: %s"), StringCast<TCHAR>(message).Get());
 			break;
 		case OAPV_LOG_INFO:
-			UE_LOG(LogApvMedia, Log, TEXT("OAPV: %s"), StringCast<TCHAR>(message).Get());
+			UE_LOG(LogApvMedia, Log, TEXT("OpenApv: %s"), StringCast<TCHAR>(message).Get());
 			break;
 		case OAPV_LOG_DEBUG:
-			UE_LOG(LogApvMedia, Verbose, TEXT("OAPV: %s"), StringCast<TCHAR>(message).Get());
+			UE_LOG(LogApvMedia, Verbose, TEXT("OpenApv: %s"), StringCast<TCHAR>(message).Get());
 			break;
 		}
 	}
@@ -234,17 +235,19 @@ namespace UE::ApvMedia
 		did = oapvd_create(&DecoderDescriptor, &Err);
 		if (did == nullptr)
 		{
-			UE_LOG(LogApvMedia, Error, TEXT("ApvDecoderContext: cannot create OAPV decoder (err=%d)"), Err);
+			UE_LOG(LogApvMedia, Error, TEXT("ApvDecoderContext: cannot create OpenApv decoder (err=%d)"), Err);
 		}
 		else
 		{
-			UE_LOG(LogApvMedia, Log, TEXT("ApvDecoderContext: created OAPV decoder with %d worker threads"), DecoderDescriptor.threads);
+			UE_LOG(LogApvMedia, Log,
+				TEXT("ApvDecoderContext: created OpenApv decoder (v%s) with %d worker threads"),
+				StringCast<TCHAR>(oapv_version(nullptr)).Get(), DecoderDescriptor.threads);
 		}
 		
 		mid = oapvm_create(&Err);
 		if(OAPV_FAILED(Err))
 		{
-			UE_LOG(LogApvMedia, Error, TEXT("ApvDecoderContext: cannot create OAPV metadata container (err=%d)"), Err);
+			UE_LOG(LogApvMedia, Error, TEXT("ApvDecoderContext: cannot create OpenApv metadata container (err=%d)"), Err);
 		}
 	}
 		
@@ -261,4 +264,4 @@ namespace UE::ApvMedia
 			mid = nullptr;
 		}
 	}
-}
\ No newline at end of file
+}
diff --git a/Engine/Plugins/Media/TmvMedia/Source/ApvMedia/Private/Apv/ApvDecoder.h b/Engine/Plugins/Media/TmvMedia/Source/ApvMedia/Private/Apv/ApvDecoder.h
--- a/Engine/Plugins/Media/TmvMedia/Source/ApvMedia/Private/Apv/ApvDecoder.h
+++ b/Engine/Plugins/Media/TmvMedia/Source/ApvMedia/Private/Apv/ApvDecoder.h
@@ -3,8 +3,7 @@
 #pragma once
 
 #include "ApvCommon.h"
-#include "Async/Mutex.h"
-#include "Templates/SharedPointer.h"
+#include "HAL/UnrealMemory.h"
 
 namespace UE::ApvMedia
 {
diff --git a/Engine/Plugins/Media/TmvMedia/Source/ApvMedia/Private/Apv/ApvEncoder.cpp b/Engine/Plugins/Media/TmvMedia/Source/ApvMedia/Private/Apv/ApvEncoder.cpp
new file mode 100644
--- /dev/null
+++ b/Engine/Plugins/Media/TmvMedia/Source/ApvMedia/Private/Apv/ApvEncoder.cpp
@@ -0,0 +1,236 @@
+// Copyright Epic Games, Inc. All Rights Reserved.
+
+#include "ApvEncoder.h"
+#include "ApvMediaLog.h"
+#include "ApvMediaSettings.h"
+#include "ApvMediaTmvEncoderOptions.h"
+#include "ApvMipBufferUtils.h"
+#include "Encoder/ITmvMediaEncoder.h"
+#include "Transcoder/TmvMediaFrameMips.h"
+
+namespace UE::ApvMedia
+{
+	int32 GetProfileIdc(EApvMediaProfile InProfile)
+	{
+		switch (InProfile)
+		{
+		case EApvMediaProfile::YCbCr400_10:
+			return OAPV_PROFILE_400_10;
+		case EApvMediaProfile::YCbCr422_10:
+			return OAPV_PROFILE_422_10;
+		case EApvMediaProfile::YCbCr422_12:
+			return OAPV_PROFILE_422_12;
+		case EApvMediaProfile::YCbCr444_10:
+			return OAPV_PROFILE_444_10;
+		case EApvMediaProfile::YCbCr444_12:
+			return OAPV_PROFILE_444_12;
+		case EApvMediaProfile::YCbCr4444_10:
+			return OAPV_PROFILE_4444_10;
+		case EApvMediaProfile::YCbCr4444_12:
+			return OAPV_PROFILE_4444_12;
+		default:
+			return OAPV_PROFILE_422_10;
+		}
+	}
+
+	int32 GetPreset(EApvMediaPreset InPreset)
+	{
+		switch (InPreset)
+		{
+		case EApvMediaPreset::Fastest:
+			return OAPV_PRESET_FASTEST;
+		case EApvMediaPreset::Fast:
+			return OAPV_PRESET_FAST;
+		case EApvMediaPreset::Medium:
+			return OAPV_PRESET_MEDIUM;
+		case EApvMediaPreset::Slow:
+			return OAPV_PRESET_SLOW;
+		case EApvMediaPreset::Placebo:
+			return OAPV_PRESET_PLACEBO;
+		}
+		return OAPV_PRESET_DEFAULT;
+	}
+
+	int32 GetValidatedNumThreads(int32 InNumThreads)
+	{
+		int32 NumThreads = InNumThreads > 0 ? InNumThreads : FPlatformMisc::NumberOfCores();
+		if (NumThreads > OAPV_MAX_THREADS)
+		{
+			UE_LOG(LogApvMedia, Warning, TEXT("Requested number of threads (%d) exceeds maximum OpenAPV threads (%d)."), NumThreads, OAPV_MAX_THREADS);
+		}
+		return FMath::Clamp(NumThreads, 1, OAPV_MAX_THREADS);
+	}
+
+	/** Log helper for displaying the name of an enum value as string. */
+	template <typename InEnumType>
+	FString GetEnumNameStringByValue(const int64 InValue)
+	{
+		FString ValueName = StaticEnum<InEnumType>()->GetNameStringByValue(InValue);
+		return !ValueName.IsEmpty() ? ValueName : TEXT("<invalid_enum_value>");
+	}
+
+	FApvEncoderContext::FApvEncoderContext(
+		const FApvMediaTmvEncoderOptions& InOptions,
+		const FTmvMediaEncoderMipInfo& InMipInfo,
+		const FTmvMediaFrameTimeInfo& InTimeInfo)
+	{
+		// Remark:
+		// we need to have all the encoder options when creating the encoder.
+		// This includes the frame dimensions and number of mips.
+		oapve_cdesc_t EncoderDescriptor;
+		FMemory::Memzero(&EncoderDescriptor, sizeof(oapve_cdesc_t));
+
+		// Primary Frame Parameters (other frames will have the same config)
+		oapve_param_t& PrimaryFrameParams = EncoderDescriptor.param[0];
+		oapve_param_default(&PrimaryFrameParams);
+
+		// Report invalid frame dimensions.
+		if (InMipInfo.Width <= 0 || InMipInfo.Height <= 0)
+		{
+			UE_LOG(LogApvMedia, Error, TEXT("ApvEncoderContext: invalid frame dimensions %dx%d."), InMipInfo.Width, InMipInfo.Height);
+			return;
+		}
+
+		// Setup encoder options here.
+		PrimaryFrameParams.w = InMipInfo.Width;
+		PrimaryFrameParams.h = InMipInfo.Height;
+		PrimaryFrameParams.tile_w = InOptions.TileSize.X;
+		PrimaryFrameParams.tile_h = InOptions.TileSize.Y;
+		PrimaryFrameParams.profile_idc = GetProfileIdc(InOptions.Profile);
+		PrimaryFrameParams.band_idc = FMath::Clamp(InOptions.Band, 0, 3);
+		PrimaryFrameParams.preset = GetPreset(InOptions.Preset);
+
+		// Note: frame rate is not stored by OpenApv frame header.
+		// todo: Could use meta data for this as an option (if no media container is used).
+		const FFrameRate FrameRate = UE::TmvMedia::TimeUtils::GetFrameRate(InTimeInfo.FrameDuration, FFrameRate(24, 1));
+		PrimaryFrameParams.fps_den = FrameRate.Denominator;
+		PrimaryFrameParams.fps_num = FrameRate.Numerator;
+
+		// Rate control - Average bitrate vs constant qp.
+		// Average bitrate: uses the frame rate to determine a target bit rate, see oapve_family_bitrate.
+		PrimaryFrameParams.rc_type = OAPV_RC_ABR;	// rate control type - average bit rate
+		// If target bitrate is not specified, the max coding rate of the specified band (and level) is used.
+
+		if (InMipInfo.ColorInfo.Encoding != UE::Color::EEncoding::None 
+			|| InMipInfo.ColorInfo.ColorSpace != UE::Color::EColorSpace::None
+			|| InMipInfo.ColorInfo.YuvMatrix != ETmvMediaFrameColorMatrix::None)
+		{
+			PrimaryFrameParams.color_description_present_flag = 1;	// enable writing color description in header.
+
+			const TOptional<EApvMediaColorTransfer> ApvColorTransfer = GetApvColorTransfer(InMipInfo.ColorInfo.Encoding);
+			if (ApvColorTransfer.IsSet())
+			{
+				PrimaryFrameParams.transfer_characteristics = ToUnderlyingType(ApvColorTransfer.GetValue());
+			}
+			else
+			{
+				UE_LOG(LogApvMedia, Warning,
+					TEXT("Encoder: Requested transfer function \"%s\" doesn't exist in Apv, will remain \"unspecified\"."),
+					*GetEnumNameStringByValue<ETextureSourceEncoding>(ToUnderlyingType(InMipInfo.ColorInfo.Encoding)));
+			}
+
+			const TOptional<EApvMediaColorSpace> ApvColorSpace = GetApvColorSpace(InMipInfo.ColorInfo.ColorSpace);
+			if (ApvColorSpace.IsSet())
+			{
+				PrimaryFrameParams.color_primaries = ToUnderlyingType(ApvColorSpace.GetValue());
+			}
+			else
+			{
+				UE_LOG(LogApvMedia, Warning,
+					TEXT("Encoder: Requested color space \"%s\" doesn't exist in Apv, will remain \"unspecified\"."),
+					*GetEnumNameStringByValue<ETextureColorSpace>(ToUnderlyingType(InMipInfo.ColorInfo.ColorSpace)));
+			}
+
+			PrimaryFrameParams.matrix_coefficients = ToUnderlyingType(GetApvColorMatrix(InMipInfo.ColorInfo.YuvMatrix));
+			PrimaryFrameParams.full_range_flag = ToUnderlyingType(GetApvColorMatrixRange(InMipInfo.ColorInfo.YuvMatrixRange));
+		}
+
+		int32 NumMips = 0;
+
+		// If the primary frame is already at the smallest possible size, don't add any extra mip frames.
+		if (InMipInfo.bEnableMips && (PrimaryFrameParams.w > 1 || PrimaryFrameParams.h > 1))
+		{
+			const EApvMediaChromaFormat ChromaFormat = InOptions.GetChromaFormat();
+
+			for (int32 MipIndex = 0, Width = FMath::Max(PrimaryFrameParams.w/2, 1), Height = FMath::Max(PrimaryFrameParams.h/2, 1);; MipIndex++)
+			{
+				// Stop once we've reached the smallest mip. 
+				if (Width == 1 && Height == 1 && MipIndex > 0)
+				{
+					break;
+				}
+
+				if (GetSubWidthC(ChromaFormat) == 2 && Width & 0x1)
+				{
+					UE_LOG(LogApvMedia, Warning, TEXT("Can't generate mip of width %d. Not a multiple of two (YUC 422/420 constraint)"), Width);
+					break;
+				}
+				if (GetSubHeightC(ChromaFormat) == 2 && Height & 0x1)
+				{
+					UE_LOG(LogApvMedia, Warning, TEXT("Can't generate mip of height %d. Not a multiple of two (YUC 420 constraint)"), Height);
+					break;
+				}
+
+				const int32 FrameIndex = 1 + MipIndex; // Assumes one primary frame.
+
+				if (FrameIndex >= OAPV_MAX_NUM_FRAMES) 
+				{ 
+					UE_LOG(LogApvMedia, Warning, TEXT("Reached OpenApv maximum number of frames (%d); stopping mip generation."), OAPV_MAX_NUM_FRAMES); 
+					break; 
+				}
+
+				EncoderDescriptor.param[FrameIndex] = PrimaryFrameParams;	// Copy primary frame parameters
+				EncoderDescriptor.param[FrameIndex].w = Width;
+				EncoderDescriptor.param[FrameIndex].h = Height;
+				++NumMips;
+
+				Width = FMath::Max(Width/2, 1);
+				Height = FMath::Max(Height/2, 1);
+			}
+		}
+
+		EncoderDescriptor.max_num_frms = 1 + NumMips;
+		EncoderDescriptor.threads = GetValidatedNumThreads(InOptions.NumThreads);
+		EncoderDescriptor.max_bs_buf_size = GetDefault<UApvMediaSettings>()->MaxEncoderBitStreamBufferSize;
+
+		// Report invalid bitstream buffer size.
+		if (EncoderDescriptor.max_bs_buf_size <= 0)
+		{
+			UE_LOG(LogApvMedia, Error, TEXT("ApvEncoderContext: Invalid MaxEncoderBitStreamBufferSize (%d). Recommended 128 MB."), EncoderDescriptor.max_bs_buf_size);
+			return;
+		}
+
+		int Err = 0;
+		eid = oapve_create(&EncoderDescriptor, &Err);
+		if (eid == nullptr)
+		{
+			UE_LOG(LogApvMedia, Error, TEXT("ApvEncoderContext: cannot create OpenApv encoder (err=%d)"), Err);
+		}
+		else
+		{
+			UE_LOG(LogApvMedia, Log,
+				TEXT("ApvEncoderContext: created OpenApv encoder (v%s) with %d worker threads"),
+				StringCast<TCHAR>(oapv_version(nullptr)).Get(), EncoderDescriptor.threads);
+		}
+
+		mid = oapvm_create(&Err);
+		if(OAPV_FAILED(Err))
+		{
+			UE_LOG(LogApvMedia, Error, TEXT("ApvEncoderContext: cannot create OpenApv metadata container (err=%d)"), Err);
+		}
+	}
+
+	FApvEncoderContext::~FApvEncoderContext()
+	{
+		if (eid != nullptr)
+		{
+			oapve_delete(eid);
+			eid = nullptr;
+		}
+		if (mid != nullptr)
+		{
+			oapvm_delete(mid);
+			mid = nullptr;
+		}
+	}
+}
diff --git a/Engine/Plugins/Media/TmvMedia/Source/ApvMedia/Private/ApvMediaModule.cpp b/Engine/Plugins/Media/TmvMedia/Source/ApvMedia/Private/ApvMediaModule.cpp
--- a/Engine/Plugins/Media/TmvMedia/Source/ApvMedia/Private/ApvMediaModule.cpp
+++ b/Engine/Plugins/Media/TmvMedia/Source/ApvMedia/Private/ApvMediaModule.cpp
@@ -2,12 +2,14 @@
 
 #include "ApvMediaLog.h"
 #include "ApvMediaTmvDecoderFactory.h"
+#include "ApvMediaTmvEncoderFactory.h"
 #include "ITmvMediaModule.h"
 #include "Modules/ModuleInterface.h"
 #include "Modules/ModuleManager.h"
 
 DEFINE_LOG_CATEGORY(LogApvMedia);
 
+/** Implementation of ApvMedia (OpenApv codec) Module. */
 class FApvMediaModule : public IModuleInterface
 {
 public:
@@ -16,6 +18,8 @@ public:
 	{
 		DecoderFactory = MakeShared<FApvMediaTmvDecoderFactory>();
 		ITmvMediaModule::GetOrLoad().RegisterDecoderFactory(DecoderFactory);
+		EncoderFactory = MakeShared<FApvMediaTmvEncoderFactory>();
+		ITmvMediaModule::GetOrLoad().RegisterEncoderFactory(EncoderFactory);
 	}
 
 	virtual void ShutdownModule() override
@@ -23,12 +27,19 @@ public:
 		if (ITmvMediaModule* TmvMediaModule = ITmvMediaModule::Get())
 		{
 			TmvMediaModule->UnregisterDecoderFactory(DecoderFactory);
+			TmvMediaModule->UnregisterEncoderFactory(EncoderFactory);
 		}
 		DecoderFactory.Reset();
+		EncoderFactory.Reset();
 	}
 	//~ End IModuleInterface
+
 private:
+	// Tmv Decoder factory for OpenApv codec.
 	TSharedPtr<FApvMediaTmvDecoderFactory, ESPMode::ThreadSafe> DecoderFactory;
+
+	// Tmv Encoder factory for OpenApv codec.
+	TSharedPtr<FApvMediaTmvEncoderFactory, ESPMode::ThreadSafe> EncoderFactory;
 };
 
 IMPLEMENT_MODULE(FApvMediaModule, ApvMedia);
\ No newline at end of file
diff --git a/Engine/Plugins/Media/TmvMedia/Source/ApvMedia/Private/ApvMediaSettings.h b/Engine/Plugins/Media/TmvMedia/Source/ApvMedia/Private/ApvMediaSettings.h
--- a/Engine/Plugins/Media/TmvMedia/Source/ApvMedia/Private/ApvMediaSettings.h
+++ b/Engine/Plugins/Media/TmvMedia/Source/ApvMedia/Private/ApvMediaSettings.h
@@ -23,4 +23,14 @@ public:
 	 */
 	UPROPERTY(config, EditAnywhere, Category = Decoder)
 	uint32 DecoderThreads = 0;
+
+	/**
+	 * Maximum size for the encoder bit stream buffer.
+	 * 
+	 * OpenApv requires specification of a maximum bit stream buffer size since the encoder
+	 * doesn't reallocate its bit stream buffer. We must specify enough space to store the
+	 * whole encoded access unit including all mips.
+	 */
+	UPROPERTY(config, EditAnywhere, Category = Encoder)
+	uint32 MaxEncoderBitStreamBufferSize = 128*1024*1024;
 };
diff --git a/Engine/Plugins/Media/TmvMedia/Source/ApvMedia/Private/ApvMediaTmvDecoder.cpp b/Engine/Plugins/Media/TmvMedia/Source/ApvMedia/Private/ApvMediaTmvDecoder.cpp
--- a/Engine/Plugins/Media/TmvMedia/Source/ApvMedia/Private/ApvMediaTmvDecoder.cpp
+++ b/Engine/Plugins/Media/TmvMedia/Source/ApvMedia/Private/ApvMediaTmvDecoder.cpp
@@ -3,6 +3,7 @@
 #include "ApvMediaTmvDecoder.h"
 
 #include "Apv/ApvDecoder.h"
+#include "Apv/ApvMipBufferUtils.h"
 #include "Apv/ApvParser.h"
 #include "ApvMediaLog.h"
 #include "Math/IntRect.h"
@@ -54,70 +55,13 @@ namespace UE::ApvMedia
 			if (InSize != 0 && InCount != 0)
 			{
 				int64 ReadBytes = reinterpret_cast<ITmvMediaDecoderAccessUnit*>(InInputStream->data)->Read(OutBuffer, InSize * InCount);
-				return ReadBytes/InSize;	// Return the number of elements (match fread return value).
+				return ReadBytes >= 0 ? static_cast<size_t>(ReadBytes)/InSize : 0;	// Return the number of elements (match fread return value).
 			}
 			return 0;
 		}
 	};
 #endif
-
-	// Determines the original color format from the setup of the mip color planes.
-	int32 GetApvColorFormat(const FTmvMediaFrameMipInfo& InMipInfo)
-	{
-		switch (InMipInfo.Planes.Num())
-		{
-		case 1:
-			return OAPV_CF_YCBCR400;
-		case 2:
-			return OAPV_CF_PLANAR2;
-		case 3:
-			// Deduce subsampling: 420, 422 0r 444
-			if (InMipInfo.Planes[1].Height == InMipInfo.Height) // 422, 444
-			{
-				if (InMipInfo.Planes[1].Width == InMipInfo.Width)
-				{
-					return OAPV_CF_YCBCR444;
-				}
-				return OAPV_CF_YCBCR422;	
-			}
-			return OAPV_CF_YCBCR420;
-		case 4:
-			return OAPV_CF_YCBCR4444;
-		default:
-			UE_LOG(LogApvMedia, Error, TEXT("FApvImageBitmap: unsupported color format"));
-			return OAPV_CF_UNKNOWN;
-		}
-	}
-
-	uint8 GetBitDepth(const FTmvMediaFrameMipInfo& InMipInfo)
-	{
-		return !InMipInfo.Planes.IsEmpty() ? InMipInfo.Planes[0].BitDepth : 0;
-	}
-
-	// Determines the original color space from the setup of the mip color planes.
-	int32 GetApvColorSpace(const FTmvMediaFrameMipInfo& InMipInfo)
-	{
-		return OAPV_CS_SET(GetApvColorFormat(InMipInfo), GetBitDepth(InMipInfo), 0);
-	}
 	
-	oapv_imgb_t* ApvSetupOutputBuffer(const FTmvMediaFrameMipInfo& InMipInfo, FTmvMediaFrameMipBuffer& InResourceBuffer)
-	{
-		TUniquePtr<FApvImageBitmap> OutFrame = MakeUnique<FApvImageBitmap>();		
-		if (OutFrame->Init(InMipInfo.Width, InMipInfo.Height, GetApvColorSpace(InMipInfo), /*bAllocate*/ false))
-		{
-			for (int i = 0; i < OutFrame->np; ++i)
-			{
-				OutFrame->a[i] = InResourceBuffer.GetPlaneBufferForComponent(i);
-				if (!ensure(OutFrame->a[i]))
-				{
-					return nullptr;	// don't let an invalid memory frame be used.
-				}
-			}
-			return OutFrame.Release();
-		}
-		return nullptr;
-	}
-
 	void SetupTileRequest(const TArray<FIntRect>& InTileRegions, const FIntPoint& InNumTiles, oapv_mip_request& OutMipRequest)
 	{
 		OutMipRequest.num_tiles = 0;
@@ -178,6 +122,10 @@ ETmvMediaDecoderResult FApvMediaParser::ParseMipInfos(ITmvMediaDecoderAccessUnit
 		MipInfo.Width = FrameInfo.Width;
 		MipInfo.NumComponents = FrameInfo.GetNumComponents();
 		MipInfo.ColorModel = ETmvMediaFrameColorModel::YUV;
+		MipInfo.ColorInfo.Encoding = GetColorEncoding(static_cast<EApvMediaColorTransfer>(FrameHeader.FrameHeader.ColorDescription.TransferCharacteristic));
+		MipInfo.ColorInfo.ColorSpace = GetColorSpace(static_cast<EApvMediaColorSpace>(FrameHeader.FrameHeader.ColorDescription.ColorPrimaries));
+		MipInfo.ColorInfo.YuvMatrix = GetColorMatrix(static_cast<EApvMediaColorMatrix>(FrameHeader.FrameHeader.ColorDescription.MatrixCoefficients));
+		MipInfo.ColorInfo.YuvMatrixRange = GetColorMatrixRange(static_cast<EApvMediaColorMatrixRange>(FrameHeader.FrameHeader.ColorDescription.FullRangeFlag));
 		MipInfo.TileWidth = FrameHeader.FrameHeader.TileInfo.GetTileWidthInSamples();
 		MipInfo.TileHeight = FrameHeader.FrameHeader.TileInfo.GetTileHeightInSamples();
 		MipInfo.NumTiles = FrameHeader.FrameHeader.GetNumTiles();
@@ -233,10 +181,10 @@ ETmvMediaDecoderResult FApvMediaTmvDecoder::Decode(ITmvMediaDecoderAccessUnit& I
 	TArray<oapv_mip_request> ApvMipRequests;
 	ApvMipRequests.Reserve(InMipRequests.Num());
 
-	// Prepare the mip requests for OpenAPV.
+	// Prepare the mip requests for OpenApv.
 	for (const FTmvMediaDecoderMipRequest& MipRequest : InMipRequests)
 	{
-		if (oapv_imgb_t* OutputBuffer = ApvSetupOutputBuffer(MipRequest.MipInfo, *MipRequest.MipBuffer))
+		if (oapv_imgb_t* OutputBuffer = ApvSetupImageBuffer(MipRequest.MipInfo, *MipRequest.MipBuffer))
 		{
 			oapv_mip_request& ApvMipRequest = ApvMipRequests.Emplace_GetRef();
 			FMemory::Memzero(&ApvMipRequest, sizeof(oapv_mip_request));
diff --git a/Engine/Plugins/Media/TmvMedia/Source/ApvMedia/Private/ApvMediaTmvEncoderOptions.h b/Engine/Plugins/Media/TmvMedia/Source/ApvMedia/Private/ApvMediaTmvEncoderOptions.h
new file mode 100644
--- /dev/null
+++ b/Engine/Plugins/Media/TmvMedia/Source/ApvMedia/Private/ApvMediaTmvEncoderOptions.h
@@ -0,0 +1,341 @@
+// Copyright Epic Games, Inc. All Rights Reserved.
+
+#pragma once
+
+#include "ApvMediaTmvEncoderFactory.h"
+#include "Encoder/TmvMediaEncoderOptions.h"
+#include "Misc/FrameRate.h"
+
+#include "ApvMediaTmvEncoderOptions.generated.h"
+
+/**
+ * List of color/chroma image buffer formats used by OpenApv.
+ */
+UENUM()
+enum class EApvMediaChromaFormat : uint8
+{
+	/** 1 component (monochrome) */
+	YCbCr400 = 0,
+	/** 3 components YUV 422 */
+	YCbCr422 = 2,
+	/** 3 components YUV 444 */
+	YCbCr444 = 3,
+	/** 4 components YUV 444 + alpha */
+	YCbCrA4444 = 4,
+	/** 3 components YUV 422 with UV in the same memory plane. */
+	YCbCr422_P2 = 5	
+};
+
+/** Optimization level control presets */
+UENUM()
+enum class EApvMediaPreset : uint8
+{
+	Fastest = 0,
+	Fast = 1,
+	Medium = 2,
+	Slow = 3,
+	Placebo = 4
+};
+
+/**
+ * Each profile specifies a subset of algorithmic features and limits that must be supported
+ * by all decoders conforming to that profile. */
+UENUM()
+enum class EApvMediaProfile : uint8
+{
+	YCbCr400_10 UMETA(DisplayName = "400-10: YCbCr400 (monochrome) 10 bits"),
+	YCbCr422_10 UMETA(DisplayName = "422-10: YCbCr422 10 bits (default)"),
+	YCbCr422_12 UMETA(DisplayName = "422-12: YCbCr422 12 bits"),
+	YCbCr444_10 UMETA(DisplayName = "444-10: YCbCr444 10 bits"),
+	YCbCr444_12 UMETA(DisplayName = "444-12: YCbCr444 12 bits"),
+	YCbCr4444_10 UMETA(DisplayName = "4444-10: YCbCrA4444 10 bits"),
+	YCbCr4444_12 UMETA(DisplayName = "4444-12: YCbCrA4444 12 bits"),
+};
+
+/**
+ * Supported bit depths.
+ */
+UENUM()
+enum class EApvMediaBitDepth : uint8
+{
+	/** 10 bits per component */
+	BitDepth_10 = 10,
+	/** 12 bits per component */
+	BitDepth_12 = 12,
+};
+
+/** 
+ * Color transfer functions defined in OpenApv.
+ * Reference: Rec H.273 https://www.itu.int/rec/T-REC-H.273 
+ */
+UENUM()
+enum class EApvMediaColorTransfer : uint8
+{
+	Reserved0 = 0 UMETA(Hidden), 
+	Bt709 = 1,
+	Unspecified = 2,
+	Reserved3 = 3 UMETA(Hidden),
+	Bt470M = 4,
+	Bt470BG = 5,
+	Smpte170m = 6,
+	Smpte240m = 7,
+	Linear = 8,
+	Log100 = 9,
+	Log316 = 10,
+	Iec61966_2_4 = 11,
+	Bt1361e = 12,
+	Iec61966_2_1 = 13,
+	Bt2020_10 = 14,
+	Bt2020_12 = 15,
+	Smpte2084 = 16,
+	Smpte428 = 17,
+	Arib_std_b67 = 18,
+};
+
+/**
+ * Color primaries defined in OpenApv.
+ * Reference: Rec H.273 https://www.itu.int/rec/T-REC-H.273  
+ */
+UENUM()
+enum class EApvMediaColorSpace : uint8
+{
+	Reserved0 = 0 UMETA(Hidden),
+	Bt709 = 1,
+	Unspecified = 2,
+	Reserved3 = 3 UMETA(Hidden),
+	Bt470m = 4,
+	Bt470bg = 5,
+	Smpte170m = 6,
+	Smpte240m = 7,
+	Film = 8,
+	Bt2020 = 9,
+	Smpte428 = 10,
+	Smpte431 = 11,
+	Smpte432 = 12,
+};
+
+/** 
+ * YUV color matrices defined in OpenApv.
+ * Reference: Rec H.273 https://www.itu.int/rec/T-REC-H.273 
+ */
+UENUM()
+enum class EApvMediaColorMatrix : uint8
+{
+	Gbr = 0,
+	Bt709 = 1,
+	Unspecified = 2,
+	Reserved3 = 3 UMETA(Hidden),
+	Fcc = 4,
+	Bt470bg = 5,
+	Smpte170m = 6,
+	Smpte240m = 7,
+	Ycgco = 8,
+	Bt2020nc = 9,
+	Bt2020c = 10,
+	Smpte2085 = 11,
+	Chroma_derived_nc = 12,
+	Chroma_derived_c = 13,
+	Ictcp = 14,
+};
+
+/** 
+ * YUV color matrix ranges defined in OpenApv.
+ * Specifies the scaling and offset applied in association with the YUV matrix coefficients.
+ * Reference: Rec H.273 https://www.itu.int/rec/T-REC-H.273 
+ */
+UENUM()
+enum class EApvMediaColorMatrixRange : uint8
+{
+	/** Standard for television, broadcast and video content. Typically defines black at 16 and white at 235 (standard dependent). */
+	Limited = 0,
+	/** Standard for PC graphics and image processing. Defines black at 0 and white at 255. */
+	Full = 1
+};
+
+namespace UE::ApvMedia
+{
+	/** Get the color/chroma format from the given apv profile. */
+	inline EApvMediaChromaFormat GetChromaFormat(const EApvMediaProfile InProfile)
+	{
+		switch (InProfile)
+		{
+		case EApvMediaProfile::YCbCr400_10:
+			return EApvMediaChromaFormat::YCbCr400;
+		case EApvMediaProfile::YCbCr422_10:
+			return EApvMediaChromaFormat::YCbCr422;
+		case EApvMediaProfile::YCbCr422_12:
+			return EApvMediaChromaFormat::YCbCr422;
+		case EApvMediaProfile::YCbCr444_10:
+			return EApvMediaChromaFormat::YCbCr444;
+		case EApvMediaProfile::YCbCr444_12:
+			return EApvMediaChromaFormat::YCbCr444;
+		case EApvMediaProfile::YCbCr4444_10:
+			return EApvMediaChromaFormat::YCbCrA4444;
+		case EApvMediaProfile::YCbCr4444_12:
+			return EApvMediaChromaFormat::YCbCrA4444;
+		default:
+			return EApvMediaChromaFormat::YCbCr422;
+		}
+	}
+
+	/** Get the bit depth from the given apv profile. */
+	inline int8 GetBitDepth(const EApvMediaProfile InProfile)
+	{
+		switch (InProfile)
+		{
+		case EApvMediaProfile::YCbCr400_10:
+			return 10;
+		case EApvMediaProfile::YCbCr422_10:
+			return 10;
+		case EApvMediaProfile::YCbCr422_12:
+			return 12;
+		case EApvMediaProfile::YCbCr444_10:
+			return 10;
+		case EApvMediaProfile::YCbCr444_12:
+			return 12;
+		case EApvMediaProfile::YCbCr4444_10:
+			return 10;
+		case EApvMediaProfile::YCbCr4444_12:
+			return 12;
+		default:
+			return 10;
+		}
+	}
+
+	/** Get the number of components from the given apv chroma format. */
+	inline int32 GetNumComponents(const EApvMediaChromaFormat InFormat)
+	{
+		switch (InFormat)
+		{
+		case EApvMediaChromaFormat::YCbCr400:
+			return 1;
+		case EApvMediaChromaFormat::YCbCr422:
+		case EApvMediaChromaFormat::YCbCr422_P2:
+		case EApvMediaChromaFormat::YCbCr444:
+			return 3;
+		case EApvMediaChromaFormat::YCbCrA4444:
+			return 4;
+		default:
+			return 3;
+		}
+	}
+
+	/** Get the number of memory planes for the given color/chroma format. */
+	inline int32 GetNumPlanes(const EApvMediaChromaFormat InFormat)
+	{
+		switch (InFormat)
+		{
+		case EApvMediaChromaFormat::YCbCr400:
+			return 1;
+		case EApvMediaChromaFormat::YCbCr422_P2:
+			return 2; // 2 planes (Y and UV)
+		case EApvMediaChromaFormat::YCbCr422:
+		case EApvMediaChromaFormat::YCbCr444:
+			return 3;
+		case EApvMediaChromaFormat::YCbCrA4444:
+			return 4;
+		default:
+			return 3;
+		}
+	}
+
+	/** Get the number of components packed in the given plane for the given format. */
+	inline int32 GetPlaneNumComponents(const int32 InPlaneIndex, const EApvMediaChromaFormat InFormat)
+	{
+		if (InFormat == EApvMediaChromaFormat::YCbCr422_P2)
+		{
+			return InPlaneIndex == 0 ? 1 : 2;
+		}
+		return 1;
+	}
+
+	/** Returns the width factor for the chroma samples of the given format. */
+	inline int32 GetSubWidthC(const EApvMediaChromaFormat InFormat)
+	{
+		// Chroma 420 is not supported.
+		if (InFormat == EApvMediaChromaFormat::YCbCr422_P2 || InFormat == EApvMediaChromaFormat::YCbCr422)
+		{
+			return 2;
+		}
+		return 1;
+	}
+
+	/** Returns the height factor for the chroma samples of the given format. */
+	inline int32 GetSubHeightC(const EApvMediaChromaFormat InFormat)
+	{
+		// Only 420 has a sub sampling for chroma.
+		// But, it doesn't seem supported by the encoder though.
+		// todo: double check.
+		return 1;
+	}
+}
+
+/**
+ * OpenApv Tmv Encoder options.
+ */
+USTRUCT(BlueprintType, DisplayName = "OpenApv Tmv Encoder Options")
+struct FApvMediaTmvEncoderOptions : public FTmvMediaEncoderOptions
+{
+	GENERATED_BODY()
+
+	FApvMediaTmvEncoderOptions()
+	{
+		// Select commonly used color settings by default.
+		bEnableColorManagement = true;
+		DestinationColorSpace = ETextureColorSpace::TCS_sRGB;
+		DestinationEncoding = ETmvMediaEncoderEncoding::sRGB;
+		YuvMatrix = ETmvMediaEncoderColorMatrix::Rec709;
+		YuvMatrixRange = ETmvMediaEncoderColorMatrixRange::Limited;
+	}
+	
+	virtual ~FApvMediaTmvEncoderOptions() override = default;
+
+	//~ Begin FTmvMediaEncoderOptions
+	virtual FName GetEncoderName() const override { return FApvMediaTmvEncoderFactory::EncoderName;}
+	virtual FString GetFileSequenceExtension() const override { return TEXT("apv1"); }
+	//~ End FTmvMediaEncoderOptions
+
+	/** Get currently configured color/chromat format. */
+	EApvMediaChromaFormat GetChromaFormat() const
+	{
+		return  UE::ApvMedia::GetChromaFormat(Profile);
+	}
+
+	/** Get currently configured bit depth */
+	int8 GetBitDepth() const
+	{
+		return  UE::ApvMedia::GetBitDepth(Profile);
+	}
+	
+	/** Dimensions of a tile in pixels. */
+	UPROPERTY(EditAnywhere, Category = "Settings", meta = (ClampMin = "0.0"))
+	FIntPoint TileSize = FIntPoint(256, 256);
+	
+	/**
+	 * Specify the color/chroma profile, includes bitdepth.
+	 * Input video will be converted to the given profile.
+	 */
+	UPROPERTY(EditAnywhere, Category = "Settings")
+	EApvMediaProfile Profile = EApvMediaProfile::YCbCr422_10;
+
+	/**
+	 * Coded data rate band Setting.
+	 * If the bitrate is not specified otherwise, this determines the max coded data rate (higher band leads to higher bit rate).
+	 */
+	UPROPERTY(EditAnywhere, Category = "Settings", meta = (ClampMin = "0", ClampMax = "3", UIMin = "0", UIMax = "3"))
+	int Band = 2;
+
+	// todo: bitrate/quality specification: low, standard, high, ultra. (Family parameter)
+
+	/**
+	 * Specify the encoding optimization level. 
+	 */
+	UPROPERTY(EditAnywhere, Category = "Settings")
+	EApvMediaPreset Preset = EApvMediaPreset::Medium;
+	
+	/**
+	 * Number of worker thread to encode tiles (0 = automatically use the available number of cores).
+	 */
+	UPROPERTY(EditAnywhere, Category = "Advanced", meta = (ClampMin = "0.0"))
+	int32 NumThreads = 0;
+};
diff --git a/Engine/Plugins/Media/TmvMedia/Source/TmvMedia/Private/SampleConverter/TmvMediaFrameMipStructuredBuffer.cpp b/Engine/Plugins/Media/TmvMedia/Source/TmvMedia/Private/SampleConverter/TmvMediaFrameMipStructuredBuffer.cpp
--- a/Engine/Plugins/Media/TmvMedia/Source/TmvMedia/Private/SampleConverter/TmvMediaFrameMipStructuredBuffer.cpp
+++ b/Engine/Plugins/Media/TmvMedia/Source/TmvMedia/Private/SampleConverter/TmvMediaFrameMipStructuredBuffer.cpp
@@ -212,7 +212,7 @@ void FTmvMediaFrameMipStructuredBuffer::CopyTileRegions(int32 InFrameId, const F
 	});
 }
 
-bool FTmvMediaFrameMipStructuredBuffer::IsValid() const
+bool FTmvMediaFrameMipStructuredBuffer::IsValidForRendering() const
 {
 	if (!UploadBufferRef.IsValid())
 	{
diff --git a/Engine/Plugins/Media/TmvMedia/Source/TmvMedia/Private/SampleConverter/TmvMediaFrameMipStructuredBuffer.h b/Engine/Plugins/Media/TmvMedia/Source/TmvMedia/Private/SampleConverter/TmvMediaFrameMipStructuredBuffer.h
--- a/Engine/Plugins/Media/TmvMedia/Source/TmvMedia/Private/SampleConverter/TmvMediaFrameMipStructuredBuffer.h
+++ b/Engine/Plugins/Media/TmvMedia/Source/TmvMedia/Private/SampleConverter/TmvMediaFrameMipStructuredBuffer.h
@@ -25,7 +25,11 @@ public:
 	virtual void* GetPlaneBufferForComponent(int32 InComponentIndex) override;
 	virtual FShaderResourceViewRHIRef GetShaderResourceView(int32 InComponentIndex) const override;
 	virtual void CopyTileRegions(int32 InFrameId, const FTmvMediaFrameMipInfo& InTmvMipInfo, TConstArrayView<FIntRect> InTileRegions, TConstArrayView<FTmvMediaShaderTileDesc> InTileInfos) override;
-	virtual bool IsValid() const override;
+	virtual bool IsValidForRendering() const override;
+	virtual bool IsStructuredBuffer() const override
+	{
+		return true;
+	}
 	//~ End FTmvMediaFrameMipBuffer
 
 private:
diff --git a/Engine/Plugins/Media/TmvMedia/Source/TmvMedia/Private/SampleConverter/TmvMediaFrameMipTextureBuffer.cpp b/Engine/Plugins/Media/TmvMedia/Source/TmvMedia/Private/SampleConverter/TmvMediaFrameMipTextureBuffer.cpp
--- a/Engine/Plugins/Media/TmvMedia/Source/TmvMedia/Private/SampleConverter/TmvMediaFrameMipTextureBuffer.cpp
+++ b/Engine/Plugins/Media/TmvMedia/Source/TmvMedia/Private/SampleConverter/TmvMediaFrameMipTextureBuffer.cpp
@@ -5,6 +5,7 @@
 #include "ProfilingDebugging/RealtimeGPUProfiler.h"
 #include "RHICommandList.h"
 #include "TmvMediaLog.h"
+#include "Utils/TmvMediaFrameUtils.h"
 
 DECLARE_GPU_STAT_NAMED(TmvConverter_AllocateTexture, TEXT("TmvConverter.AllocateTexture"));
 DECLARE_GPU_STAT_NAMED(TmvConverter_UploadTextureRegions, TEXT("TmvConverter.UploadTextureRegions"));
@@ -17,130 +18,8 @@ FTmvMediaFrameMipTextureBuffer::~FTmvMediaFrameMipTextureBuffer()
 {
 }
 
-namespace UE::TmvConverterTextureBuffer::Private
+namespace UE::TmvMedia
 {
-	/**
-	 * Pre-defined pixel format correspondence matrices matching memory layouts for packed components.
-	 */
-	namespace PixelFormatMatrices
-	{
-		constexpr EPixelFormat Skip = PF_Unknown;		// no DXGI/Hardware format exists (for any 3 bytes per component or 8 bits float).
-		constexpr EPixelFormat Missing = PF_Unknown;	// No PF defined in UE, but exists in DXGI.
-	
-		/**
-		 * Matrix of float pixel formats.
-		 * Row: Number of Components, Column: Number of Bytes Per Components
-		 */
-		static const EPixelFormat FloatPixelFormats[4][4] =
-		{
-			/* 1 component(s) */ {Skip, PF_R16F, Skip, PF_R32_FLOAT},
-			/* 2 component(s) */ {Skip, PF_G16R16F, Skip, PF_G32R32F},	// Remark: GR instead of RG. @todo Will need shader swizzle.
-			/* 3 component(s) */ {Skip, Skip, Skip, PF_R32G32B32F},		// No packed 24 or 48 bits format exits.
-			/* 4 component(s) */ {Skip, PF_FloatRGBA, Skip, Missing}
-		};
-
-		/**
-		 * Matrix of normalized integer pixel formats.
-		 * Row: Number of Components, Column: Number of Bytes Per Components
-		 */
-		static const EPixelFormat NormPixelFormats[4][4] =
-		{
-			// Note: all 32 bits formats are missing. Not commonly used for video formats (yet).
-			/* 1 component(s) */ {PF_G8, PF_G16, Skip, Missing},
-			/* 2 component(s) */ {PF_R8G8, PF_G16R16, Skip, Missing},
-			/* 3 component(s) */ {Skip, Skip, Skip, Missing},			// No packed 24 or 48 bits format exits.
-			/* 4 component(s) */ {PF_R8G8B8A8, PF_R16G16B16A16_UNORM, Skip, Missing}
-		};
-
-		/**
-		 * Matrix of integer pixel formats.
-		 * Row: Number of Components, Column: Number of Bytes Per Components
-		 * @remark Can only be used in UAV for compute shaders. (todo)
-		 */
-		static const EPixelFormat IntPixelFormats[4][4] =
-		{
-			/* 1 component(s) */ {PF_R8_UINT, PF_R16_UINT, Skip, PF_R32_UINT},
-			/* 2 component(s) */ {PF_R8G8_UINT, PF_R16G16_UINT, Skip, PF_R32G32_UINT},
-			/* 3 component(s) */ {Skip, Skip, Skip, PF_R32G32B32_UINT},	// No packed 24 or 48 bits format exits.
-			/* 4 component(s) */ {PF_R8G8B8A8_UINT, PF_R16G16B16A16_UINT, Skip, PF_R32G32B32A32_UINT}
-		};
-		
-		EPixelFormat LookupPixelFormat(ETmvMediaFrameComponentType Type, bool bInNormalized, int32 ComponentIndex, int32 ComponentSizeIndex)
-		{
-			if (Type == ETmvMediaFrameComponentType::Float)
-			{
-				return FloatPixelFormats[ComponentIndex][ComponentSizeIndex];
-			}
-
-			if (bInNormalized)
-			{
-				// For Pixel shaders, we can only use UNORM formats. Note that all 32 bits formats are missing.
-				return NormPixelFormats[ComponentIndex][ComponentSizeIndex];
-			}
-
-			// For Compute shaders, we can use _UINT. This is the most complete set of formats.
-			return IntPixelFormats[ComponentIndex][ComponentSizeIndex];
-		}
-	}
-
-	EPixelFormat GetPackedPlanePixelFormat(const FTmvMediaFramePlaneInfo& InPlaneInfo, bool bInNormalized)
-	{
-		check(InPlaneInfo.ComponentLayout == ETmvMediaFrameComponentLayout::Packed);
-
-		if (InPlaneInfo.NumComponents < 1 || InPlaneInfo.NumComponents > 4)
-		{
-			UE_LOG(LogTmvMedia, Verbose, TEXT("GetPackedPlanePixelFormat: Invalid Number of Components: %d, should be in [1,4]."), InPlaneInfo.NumComponents);
-			return PF_Unknown;
-		}
-
-		const int32 BytesPerComponent = InPlaneInfo.GetBytesPerComponent();
-		
-		if (BytesPerComponent < 1 || BytesPerComponent > 4)
-		{
-			UE_LOG(LogTmvMedia, Verbose, TEXT("GetPackedPlanePixelFormat: Invalid Number of bytes per components: %d, should be in [1,4]."), BytesPerComponent);
-			return PF_Unknown;
-		}
-
-		return PixelFormatMatrices::LookupPixelFormat(InPlaneInfo.Type, bInNormalized, InPlaneInfo.NumComponents - 1, BytesPerComponent - 1);
-	}
-
-	EPixelFormat GetInterleavedPlanePixelFormat(const FTmvMediaFramePlaneInfo& InPlaneInfo, bool bInNormalized)
-	{
-		check(InPlaneInfo.ComponentLayout == ETmvMediaFrameComponentLayout::Interleaved);
-
-		const int32 BytesPerComponent = InPlaneInfo.GetBytesPerComponent();
-		
-		if (BytesPerComponent < 1 || BytesPerComponent > 4)
-		{
-			UE_LOG(LogTmvMedia, Verbose, TEXT("GetInterleavedPlanePixelFormat: Invalid Number of bytes per components: %d, should be in [1,4]."), BytesPerComponent);
-			return PF_Unknown;
-		}
-
-		// If the plane is interleaved, the format is going to be that of a single component, but we will allocate more pixels.
-		constexpr int32 ComponentIndex = 0;		
-		return PixelFormatMatrices::LookupPixelFormat(InPlaneInfo.Type, bInNormalized, ComponentIndex, BytesPerComponent - 1);
-	}
-
-	/**
-	 * Get the packed pixel format exactly matching the plane's memory layout.
-	 * Mostly support 1, 2 and 4 component(s) packed formats. Mainly validated against DXGI formats.
-	 * @param InPlaneInfo Plane configuration
-	 * @param bInNormalized In case of an integer format indicate if the format should be normalized or not.
-	 * @return Request pixel format if supported, PF_Unknown if not.
-	 */
-	EPixelFormat GetPlanePixelFormat(const FTmvMediaFramePlaneInfo& InPlaneInfo, bool bInNormalized)
-	{
-		switch (InPlaneInfo.ComponentLayout)
-		{
-		case ETmvMediaFrameComponentLayout::Packed:
-			return GetPackedPlanePixelFormat(InPlaneInfo, bInNormalized);
-		case ETmvMediaFrameComponentLayout::Interleaved:
-			return GetInterleavedPlanePixelFormat(InPlaneInfo, bInNormalized);
-		default:
-			return PF_Unknown;
-		}
-	}
-
 	/**
 	 * Calculate the extent of a texture for the given plane.
 	 * @remark supports interleaved layouts through emulation with a single component texture.
@@ -157,9 +36,21 @@ namespace UE::TmvConverterTextureBuffer::Private
 
 		// For interleaved component layouts, we need to do a bit of math to
 		// figure out what should be the extent of the texture. (WIP)
+		if (InPlaneInfo.NumComponents == 0) 
+		{ 
+			UE_LOG(LogTmvMedia, Verbose, TEXT("GetPlaneTextureExtent: Invalid NumComponents=%d for interleaved plane."), InPlaneInfo.NumComponents); 
+			return FIntPoint::ZeroValue; 
+		} 
+
 		const int32 BytesPerComponent = InPlaneInfo.GetBytesPerComponent();
+		if (BytesPerComponent <= 0) 
+		{ 
+			UE_LOG(LogTmvMedia, Verbose, TEXT("GetPlaneTextureExtent: Invalid BytesPerComponent=%d for interleaved plane (BitDepth=%d)."), BytesPerComponent, InPlaneInfo.BitDepth); 
+			return FIntPoint::ZeroValue; 
+		} 
+		
 		int32 ComponentStride = InPlaneInfo.Width * BytesPerComponent;	// no padding
-		int32 ComponentStrideWithPadding = InPlaneInfo.Stride / InPlaneInfo.NumComponents;
+		int32 ComponentStrideWithPadding = InPlaneInfo.NumComponents ? InPlaneInfo.Stride / InPlaneInfo.NumComponents : InPlaneInfo.Stride;
 		int32 ComponentPadding = ComponentStrideWithPadding - ComponentStride;
 		
 		// To match the memory layout, we have to include the horizontal padding between all the components but the last.
@@ -177,7 +68,7 @@ namespace UE::TmvConverterTextureBuffer::Private
 		{
 			return false;
 		}
-		if (GetPlanePixelFormat(InPlaneInfo0, /*Normalized*/ true) != GetPlanePixelFormat(InPlaneInfo1, /*Normalized*/ true))
+		if (FrameUtils::GetPlanePixelFormat(InPlaneInfo0, /*Normalized*/ true) != FrameUtils::GetPlanePixelFormat(InPlaneInfo1, /*Normalized*/ true))
 		{
 			return false;
 		}
@@ -261,7 +152,7 @@ namespace UE::TmvConverterTextureBuffer::Private
 
 bool FTmvMediaFrameMipTextureBuffer::RequestAllocation(const FTmvMediaFrameMipInfo& InMipInfo)
 {
-	using namespace UE::TmvConverterTextureBuffer::Private;
+	using namespace UE::TmvMedia;
 
 	MipInfo = InMipInfo;
 	AllocatedBufferSize = InMipInfo.GetMemorySizeInBytes();
@@ -269,8 +160,8 @@ bool FTmvMediaFrameMipTextureBuffer::RequestAllocation(const FTmvMediaFrameMipIn
 	PlaneBuffers.Reset();
 	PlaneBuffers.Reserve(InMipInfo.Planes.Num());
 
-	TArray<TSharedPtr<FPlaneTextureResource>> NewPlaneTextures;
-	NewPlaneTextures.Reserve(InMipInfo.Planes.Num());
+	TSharedPtr<FPlaneTextureResources> NewPlaneTextureResources = MakeShared<FPlaneTextureResources>();
+	NewPlaneTextureResources->PlaneTextures.Reserve(InMipInfo.Planes.Num());
 
 	for (const FTmvMediaFramePlaneInfo& PlaneInfo : InMipInfo.Planes)
 	{
@@ -292,7 +183,7 @@ bool FTmvMediaFrameMipTextureBuffer::RequestAllocation(const FTmvMediaFrameMipIn
 		TSharedPtr<FPlaneTextureResource> NewPlaneTexture = MakeShared<FPlaneTextureResource>();
 		NewPlaneTexture->Info = PlaneInfo;
 		// While we may have pixel format, it may not be supported, we will only know on the render thread when trying to allocate the texture.
-		NewPlaneTexture->PixelFormat = GetPlanePixelFormat(PlaneInfo, /*Normalized*/ true);
+		NewPlaneTexture->PixelFormat = FrameUtils::GetPlanePixelFormat(PlaneInfo, /*Normalized*/ true);
 
 		if (NewPlaneTexture->PixelFormat == PF_Unknown)
 		{
@@ -320,14 +211,23 @@ bool FTmvMediaFrameMipTextureBuffer::RequestAllocation(const FTmvMediaFrameMipIn
 		}
 
 		NewPlaneTexture->CpuBuffer = NewPlaneBuffer;
-		NewPlaneTextures.Add(NewPlaneTexture);
+		NewPlaneTextureResources->PlaneTextures.Add(NewPlaneTexture);
 	}
 
+	PlaneTextureResources = NewPlaneTextureResources;
+	TWeakPtr<FPlaneTextureResources> PlaneTextureResourcesWeak = NewPlaneTextureResources;
+	
 	// Allocate gpu texture
-	ENQUEUE_RENDER_COMMAND(CreateTmvPooledTextureBuffer)([this, NewPlaneTextures](FRHICommandListImmediate& RHICmdList)
+	ENQUEUE_RENDER_COMMAND(CreateTmvPooledTextureBuffer)([PlaneTextureResourcesWeak](FRHICommandListImmediate& RHICmdList)
 		{
+			TSharedPtr<FPlaneTextureResources> PlaneTextureResources = PlaneTextureResourcesWeak.Pin();
+			if (!PlaneTextureResources)
+			{
+				return;
+			}
+			
 			const static FLazyName ClassName(TEXT("FPlaneTextureResource"));
-			for (const TSharedPtr<FPlaneTextureResource>& PlaneTexture : NewPlaneTextures)
+			for (const TSharedPtr<FPlaneTextureResource>& PlaneTexture : PlaneTextureResources->PlaneTextures)
 			{
 				const FRHITextureCreateDesc TextureDesc =
 					FRHITextureCreateDesc::Create2D(TEXT("TmvPlaneTextureResource"))
@@ -357,7 +257,6 @@ bool FTmvMediaFrameMipTextureBuffer::RequestAllocation(const FTmvMediaFrameMipIn
 					UE_LOG(LogTmvMedia, Verbose, TEXT("Failed to allocate a texture for plane. The format is probably not supported on the platform."));
 				}
 			}
-			PlaneTexturesRT = NewPlaneTextures;
 	});
 	return true;
 }
@@ -369,7 +268,7 @@ void FTmvMediaFrameMipTextureBuffer::WaitAllocation()
 
 bool FTmvMediaFrameMipTextureBuffer::TryUpdateMipInfo(const FTmvMediaFrameMipInfo& InMipInfo)
 {
-	using namespace UE::TmvConverterTextureBuffer::Private;
+	using namespace UE::TmvMedia;
 
 	if (MipInfo.Layout != InMipInfo.Layout || MipInfo.Planes.Num() != InMipInfo.Planes.Num())
 	{
@@ -409,29 +308,36 @@ void* FTmvMediaFrameMipTextureBuffer::GetPlaneBufferForComponent(int32 InCompone
 
 FShaderResourceViewRHIRef FTmvMediaFrameMipTextureBuffer::GetShaderResourceView(int32 InComponentIndex) const
 {
-	// Use the mip info to find the plane index for this component
-	const int32 PlaneIndex = MipInfo.GetPlaneIndexForComponent(InComponentIndex);
-	if (PlaneTexturesRT.IsValidIndex(PlaneIndex) && PlaneTexturesRT[PlaneIndex])
+	if (PlaneTextureResources)
 	{
-		return PlaneTexturesRT[PlaneIndex]->ShaderResourceView;
+		// Use the mip info to find the plane index for this component
+		const int32 PlaneIndex = MipInfo.GetPlaneIndexForComponent(InComponentIndex);
+		if (PlaneTextureResources->PlaneTextures.IsValidIndex(PlaneIndex) && PlaneTextureResources->PlaneTextures[PlaneIndex])
+		{
+			return PlaneTextureResources->PlaneTextures[PlaneIndex]->ShaderResourceView;
+		}
 	}
 	return nullptr;
 }
 
 void FTmvMediaFrameMipTextureBuffer::CopyTileRegions(int32 InFrameId, const FTmvMediaFrameMipInfo& InTmvMipInfo, TConstArrayView<FIntRect> InTileRegions, TConstArrayView<FTmvMediaShaderTileDesc> InTileInfos)
 {
-	using namespace UE::TmvConverterTextureBuffer::Private;
+	using namespace UE::TmvMedia;
 	
 	TArray<FBufferRegion> TmpBufferRegionsToCopy = CalculateBufferRegionsToCopy(InTmvMipInfo, InTileRegions);
-	TSharedPtr<FTmvMediaFrameMipBuffer> BufferData = AsShared();
+	TWeakPtr<FPlaneTextureResources> PlaneTextureResourcesWeak = PlaneTextureResources;
 
-	ENQUEUE_RENDER_COMMAND(CopyFromUploadBuffer)([BufferData, BufferRegionsToCopy = MoveTemp(TmpBufferRegionsToCopy), FrameId = InFrameId](FRHICommandListImmediate& InRHICmdList)
+	ENQUEUE_RENDER_COMMAND(CopyFromUploadBuffer)([PlaneTextureResourcesWeak, BufferRegionsToCopy = MoveTemp(TmpBufferRegionsToCopy), FrameId = InFrameId](FRHICommandListImmediate& InRHICmdList)
 	{
 		RHI_BREADCRUMB_EVENT_STAT_F(InRHICmdList, TmvConverter_UploadTextureRegions, "TmvMediaConverter.StartCopy", "TmvMediaConverter.StartCopy %d", FrameId);
 
-		FTmvMediaFrameMipTextureBuffer* TextureBuffer = static_cast<FTmvMediaFrameMipTextureBuffer*>(BufferData.Get());
-
-		for (const TSharedPtr<FPlaneTextureResource>& PlaneTexture : TextureBuffer->PlaneTexturesRT)
+		TSharedPtr<FPlaneTextureResources> PlaneTextureResources = PlaneTextureResourcesWeak.Pin();
+		if (!PlaneTextureResources)
+		{
+			return;
+		}
+			
+		for (const TSharedPtr<FPlaneTextureResource>& PlaneTexture : PlaneTextureResources->PlaneTextures)
 		{
 			if (PlaneTexture && PlaneTexture->Texture && PlaneTexture->CpuBuffer)
 			{
@@ -458,14 +364,14 @@ void FTmvMediaFrameMipTextureBuffer::CopyTileRegions(int32 InFrameId, const FTmv
 	});
 }
 
-bool FTmvMediaFrameMipTextureBuffer::IsValid() const
+bool FTmvMediaFrameMipTextureBuffer::IsValidForRendering() const
 {
-	if (PlaneTexturesRT.IsEmpty())
+	if (!PlaneTextureResources || PlaneTextureResources->PlaneTextures.IsEmpty())
 	{
 		return false;
 	}
 
-	for (const TSharedPtr<FPlaneTextureResource>& PlaneTexture : PlaneTexturesRT)
+	for (const TSharedPtr<FPlaneTextureResource>& PlaneTexture : PlaneTextureResources->PlaneTextures)
 	{
 		if (!PlaneTexture.IsValid() || !PlaneTexture->ShaderResourceView)
 		{
diff --git a/Engine/Plugins/Media/TmvMedia/Source/TmvMedia/Private/SampleConverter/TmvMediaFrameMipTextureBuffer.h b/Engine/Plugins/Media/TmvMedia/Source/TmvMedia/Private/SampleConverter/TmvMediaFrameMipTextureBuffer.h
--- a/Engine/Plugins/Media/TmvMedia/Source/TmvMedia/Private/SampleConverter/TmvMediaFrameMipTextureBuffer.h
+++ b/Engine/Plugins/Media/TmvMedia/Source/TmvMedia/Private/SampleConverter/TmvMediaFrameMipTextureBuffer.h
@@ -27,7 +27,7 @@ public:
 	virtual void* GetPlaneBufferForComponent(int32 InComponentIndex) override;
 	virtual FShaderResourceViewRHIRef GetShaderResourceView(int32 InComponentIndex) const override;
 	virtual void CopyTileRegions(int32 InFrameId, const FTmvMediaFrameMipInfo& InTmvMipInfo, TConstArrayView<FIntRect> InTileRegions, TConstArrayView<FTmvMediaShaderTileDesc> InTileInfos) override;
-	virtual bool IsValid() const override;
+	virtual bool IsValidForRendering() const override;
 	//~ End FTmvMediaFrameMipBuffer
 
 private:
@@ -38,8 +38,8 @@ private:
 		TArray64<uint8> Buffer;
 	};
 	TArray<TSharedPtr<FPlaneTextureBuffer>> PlaneBuffers;
-
-	/** Texture for the memory plane. */
+	
+	/** Render thread resource: Texture for the memory plane. */
 	struct FPlaneTextureResource
 	{
 		FTmvMediaFramePlaneInfo Info;
@@ -48,5 +48,13 @@ private:
 		FShaderResourceViewRHIRef ShaderResourceView;
 		TSharedPtr<FPlaneTextureBuffer> CpuBuffer;
 	};
-	TArray<TSharedPtr<FPlaneTextureResource>> PlaneTexturesRT;
+
+	/** Container for all the render thread resources. */
+	struct FPlaneTextureResources
+	{
+		TArray<TSharedPtr<FPlaneTextureResource>> PlaneTextures;
+	};
+
+	/** Reference to the render thread resources. */
+	TSharedPtr<FPlaneTextureResources> PlaneTextureResources;
 };
\ No newline at end of file
diff --git a/Engine/Plugins/Media/TmvMedia/Source/TmvMedia/Private/SampleConverter/TmvMediaTextureSampleConverter.cpp b/Engine/Plugins/Media/TmvMedia/Source/TmvMedia/Private/SampleConverter/TmvMediaTextureSampleConverter.cpp
--- a/Engine/Plugins/Media/TmvMedia/Source/TmvMedia/Private/SampleConverter/TmvMediaTextureSampleConverter.cpp
+++ b/Engine/Plugins/Media/TmvMedia/Source/TmvMedia/Private/SampleConverter/TmvMediaTextureSampleConverter.cpp
@@ -8,8 +8,9 @@
 #include "ProfilingDebugging/RealtimeGPUProfiler.h"
 #include "ScreenPass.h"
 #include "TmvMediaLog.h"
-#include "TmvTextureBufferSwizzlingShader.h"
 #include "TmvStructuredBufferSwizzlingShader.h"
+#include "TmvTextureBufferSwizzlingShader.h"
+#include "Utils/TmvMediaFrameUtils.h"
 
 DECLARE_GPU_STAT_NAMED(TmvMediaConverter, TEXT("TmvMediaConverter"));
 DECLARE_GPU_STAT_NAMED(TmvMediaConverter_MipRender, TEXT("TmvMediaConverter.MipRender"));
@@ -22,7 +23,7 @@ static FAutoConsoleVariableRef CVarTmvConverterTintMipmaps(
 	TEXT("The TMV sample converter will tint the mipmaps for visualizing (0-Red, 1-Green, 2-blue, 3-Yellow, 4+Pink).\n"),
 	ECVF_RenderThreadSafe);
 
-namespace UE::TmvMediaConverter::Private
+namespace UE::TmvMedia::Converter
 {
 	/** This function is similar to DrawScreenPass in OpenColorIODisplayExtension.cpp except it is catered for Viewless texture rendering. */
 	template<typename TSetupFunction>
@@ -56,7 +57,15 @@ namespace UE::TmvMediaConverter::Private
 
 	uint32 GetMaxValue(uint8 InBitDepth)
 	{
-		return (1u << static_cast<uint32>(InBitDepth)) - 1;
+		if (InBitDepth == 0)
+		{
+			return 0;
+		}
+		if (InBitDepth >= 32u)
+		{
+			return 0xFFFFFFFFu;
+		}
+		return (1u << static_cast<uint32>(InBitDepth)) - 1u;
 	}
 
 	/** Normalization factor for _UINT (not normalized) buffer lookup. */
@@ -64,7 +73,7 @@ namespace UE::TmvMediaConverter::Private
 	{
 		if (InType == ETmvMediaFrameComponentType::Int)
 		{
-			return 1.0f/static_cast<float>(GetMaxValue(InBitDepth));
+			return InBitDepth > 0 ? 1.0f/static_cast<float>(GetMaxValue(InBitDepth)) : 1.0f;
 		}
 		return 1.0f;
 	}
@@ -82,30 +91,26 @@ namespace UE::TmvMediaConverter::Private
 		}
 		return 1.0f;
 	}
-
-	// @todo: move to MediaShaders.
-	const FVector YUVOffset12bits = FVector(256.0f/4095.0f, 2048.0f/4095.0f, 2048.0f/4095.0f);
 	
-	FVector GetYUVOffset(ETmvMediaFrameComponentType InType, uint8 InBitDepth)
+	enum class EDataRangeType : uint8
+    {
+    	/** raw, non-negative whole number. These need to be normalized (or adjusted) to the bit depth. */
+    	UInt,
+    	/** sampled as floating point normalized between 0-1. These need to be scaled back (or denormalized) to map to bit depth. */
+    	UNorm
+    };
+	
+	float GetDataFactor(EDataRangeType InDataRangeType, ETmvMediaFrameComponentType InType, uint8 InBitDepth)
 	{
-		if (InType == ETmvMediaFrameComponentType::Int)
+		switch (InDataRangeType)
 		{
-			switch (InBitDepth)
-			{
-			case 8:
-				return MediaShaders::YUVOffset8bits;
-			case 10:
-				return MediaShaders::YUVOffset10bits;
-			case 12:
-				return YUVOffset12bits;
-			case 16:
-				return MediaShaders::YUVOffset16bits;
-			default:
-				return MediaShaders::YUVOffset16bits;
-			}
-		}
+		case EDataRangeType::UInt:
+			return GetDataNorm(InType, InBitDepth);
 
-		return MediaShaders::YUVOffsetFloat;
+		case EDataRangeType::UNorm:
+		default:
+			return GetDataScale(InType, InBitDepth);
+		}
 	}
 
 	FMatrix CombineColorTransformAndOffset(const FMatrix& InMatrix, const FVector& InYUVOffset, float InDataNorm)
@@ -146,15 +151,97 @@ namespace UE::TmvMediaConverter::Private
 		{1.0f, 1.0f, 0.0f, 0.5f}, // yellow
 		{1.0f, 0.0f, 1.0f, 0.5f}  // purple
 		};
+	
+	void SetupTintColor(int32 InMipLevel, FTmvMediaShaderColorParameters& OutColorParams)
+	{
+		if (bTmvConverterTintMipmaps)
+		{
+			int TintIdx = FMath::Min(NumMipColors-1, InMipLevel);
+			OutColorParams.MipTint = MipmapColorizing[TintIdx];
+		}
+		else
+		{
+			OutColorParams.MipTint = FLinearColor(1.0f, 1.0f, 1.0f, 0.0f);
+		}
+	}
+	
+
+	void SetupYuvColorParams( 
+		const FTmvMediaFrameMipInfo& InMipInfo,
+		EDataRangeType InDataRange,
+		FTmvMediaShaderColorParameters& OutColorParams)
+	{
+		OutColorParams.bConvertYUV = (InMipInfo.ColorModel == ETmvMediaFrameColorModel::YUV && InMipInfo.Planes.IsValidIndex(0)) ? 1u : 0u;
+		if (OutColorParams.bConvertYUV)
+		{
+			const ETmvMediaFrameComponentType ComponentType = InMipInfo.Planes[0].Type;
+			const uint8 BitDepth = InMipInfo.Planes[0].BitDepth;
+			const FVector YUVOffset = FrameUtils::GetYuvOffset(ComponentType, InMipInfo.ColorInfo.YuvMatrixRange, BitDepth);
+			const float DataFactor = GetDataFactor(InDataRange, ComponentType, BitDepth);
+			const FMatrix YuvToRgbMatrix = FrameUtils::GetYuvToRgbMatrix(InMipInfo.ColorInfo.YuvMatrix, InMipInfo.ColorInfo.YuvMatrixRange);
+			OutColorParams.YUVMatrix = static_cast<FMatrix44f>(CombineColorTransformAndOffset(YuvToRgbMatrix, YUVOffset, DataFactor));
+		}
+		else
+		{
+			OutColorParams.YUVMatrix = FMatrix44f::Identity;
+		}
+	}
+	
+	/** Configure the color related shader parameters. */
+	void SetupColorParams( 
+		const FTmvMediaTextureSampleConverterParameters& InConverterParams, 
+		const FTmvMediaFrameMipInfo& InMipInfo,
+		FTmvMediaShaderColorParameters& OutColorParams)
+	{
+		// Color settings from the decoded media (if available).
+		Color::EEncoding SourceEncoding = InMipInfo.ColorInfo.Encoding;
+		Color::FColorSpace SourceColorSpace = Color::FColorSpace(InMipInfo.ColorInfo.ColorSpace);
+		Color::EChromaticAdaptationMethod Method = Color::DEFAULT_CHROMATIC_ADAPTATION_METHOD;
+
+		// Apply potential media source color settings overrides.
+		if (InConverterParams.SourceColorSettings.IsValid())
+		{
+			if (InConverterParams.SourceColorSettings->HasEncodingOverride())
+			{
+				SourceEncoding = InConverterParams.SourceColorSettings->GetEncodingOverride();
+			}
+			SourceColorSpace = InConverterParams.SourceColorSettings->GetColorSpaceOverride(SourceColorSpace);
+			Method = InConverterParams.SourceColorSettings->GetChromaticAdaptationMethod();
+		}
+		
+		OutColorParams.bApplyColorTransform = 0u;
+		
+		// Note: we always decode to linear, if the destination RT is sRGB, it will encode to that automatically.
+		if (SourceEncoding != Color::EEncoding::Linear && SourceEncoding != Color::EEncoding::None)
+		{
+			OutColorParams.bApplyColorTransform = 1u;
+			OutColorParams.EOTF = static_cast<uint32>(SourceEncoding);
+		}
+		else
+		{
+			OutColorParams.EOTF = static_cast<uint32>(Color::EEncoding::None);
+		}
+		
+		// Destination color space is the "working" color space (should be sRGB primaries).
+		const UE::Color::FColorSpace& DestinationCS = UE::Color::FColorSpace::GetWorking();
+			
+		if (SourceColorSpace.Equals(DestinationCS))
+		{
+			OutColorParams.ColorSpaceMatrix = FMatrix44f::Identity;
+		}
+		else
+		{
+			OutColorParams.bApplyColorTransform = 1u;
+			OutColorParams.ColorSpaceMatrix = UE::Color::Transpose<float>(UE::Color::FColorSpaceTransform(SourceColorSpace, DestinationCS, Method));
+		}
+	}
 
 	void RenderMip_PlaneTextures
 		( FRHICommandListImmediate& RHICmdList
 		, const FTextureRHIRef& RenderTargetTextureRHI
 		, const FTmvMediaTextureSampleConverterParameters& ConverterParams
-		, int32 SampleMipLevel
 		, int32 TextureMipLevel
 		, const FTmvMediaFrameMipBufferHandle& BufferData
-		, const FIntPoint& SampleSize
 		, const FIntPoint& TextureSize
 		, const TArray<FIntRect>& MipViewports)
 	{
@@ -170,54 +257,10 @@ namespace UE::TmvMediaConverter::Private
 
 		FTmvTextureBufferSwizzlePS::FParameters Parameters = FTmvTextureBufferSwizzlePS::FParameters();
 
-		Parameters.AlphaScale = MipInfo.Planes.Num() == 4 ? GetDataScale(MipInfo.Planes[3].Type, MipInfo.Planes[3].BitDepth) : 1.0f;
-
-		Parameters.bConvertYUV = MipInfo.ColorModel == ETmvMediaFrameColorModel::YUV ? 1u : 0u;
-		if (Parameters.bConvertYUV)
-		{
-			const FVector YUVOffset = GetYUVOffset(MipInfo.Planes[0].Type, MipInfo.Planes[0].BitDepth); 
-			const float DataScale = GetDataScale(MipInfo.Planes[0].Type, MipInfo.Planes[0].BitDepth);
-			Parameters.YUVMatrix = static_cast<FMatrix44f>(CombineColorTransformAndOffset(MediaShaders::YuvToRgbRec709Scaled, YUVOffset, DataScale));
-		}
-		else
-		{
-			Parameters.YUVMatrix = FMatrix44f::Identity;
-		}
-
-		if (ConverterParams.SourceColorSettings.IsValid())
-		{
-			Parameters.bApplyColorTransform = 1u;
-			Parameters.EOTF = static_cast<uint32>(ConverterParams.SourceColorSettings->GetEncodingOverride());
-
-			const UE::Color::FColorSpace& DestinationCS = UE::Color::FColorSpace::GetWorking();
-			const UE::Color::FColorSpace& SourceCS = ConverterParams.SourceColorSettings->GetColorSpaceOverride(DestinationCS);
-			
-			if (SourceCS.Equals(DestinationCS))
-			{
-				Parameters.ColorSpaceMatrix = FMatrix44f::Identity;
-			}
-			else
-			{
-				const UE::Color::EChromaticAdaptationMethod Method = ConverterParams.SourceColorSettings->GetChromaticAdaptationMethod();
-				Parameters.ColorSpaceMatrix = UE::Color::Transpose<float>(UE::Color::FColorSpaceTransform(SourceCS, DestinationCS, Method));
-			}
-		}
-		else
-		{
-			Parameters.bApplyColorTransform = 0u;
-			Parameters.EOTF = 0u;
-			Parameters.ColorSpaceMatrix = FMatrix44f::Identity;
-		}
-
-		if (bTmvConverterTintMipmaps)
-		{
-			int TintIdx = FMath::Min(NumMipColors-1, TextureMipLevel);
-			Parameters.MipTint = MipmapColorizing[TintIdx];
-		}
-		else
-		{
-			Parameters.MipTint = FLinearColor(1.0f, 1.0f, 1.0f, 0.0f);
-		}
+		Parameters.ColorParams.AlphaScale = MipInfo.Planes.Num() == 4 ? GetDataScale(MipInfo.Planes[3].Type, MipInfo.Planes[3].BitDepth) : 1.0f;
+		SetupYuvColorParams(MipInfo, EDataRangeType::UNorm, Parameters.ColorParams);
+		SetupColorParams(ConverterParams, MipInfo, Parameters.ColorParams);
+		SetupTintColor(TextureMipLevel, Parameters.ColorParams);
 		
 		Parameters.InPlaneTexture0 = BufferData->GetShaderResourceView(0);
 		Parameters.InPlaneTexture1 = BufferData->GetShaderResourceView(1);
@@ -263,10 +306,11 @@ namespace UE::TmvMediaConverter::Private
 		RHICmdList.BeginRenderPass(RPInfo, TEXT("TmvTextureSwizzle"));
 
 		const FTmvMediaFrameMipInfo& MipInfo = BufferData->GetMipInfoRef();
+		ETmvMediaFrameComponentType ComponentType = MipInfo.Planes.IsValidIndex(0) ? MipInfo.Planes[0].Type : ETmvMediaFrameComponentType::Int;
 		
 		FTmvStructuredBufferSwizzlePS::FPermutationDomain PermutationVector;
 		PermutationVector.Set<FTmvStructuredBufferSwizzlePS::FNumComponents>(MipInfo.NumComponents - 1);
-		PermutationVector.Set<FTmvStructuredBufferSwizzlePS::FElementFormat>(MipInfo.Planes[0].Type == ETmvMediaFrameComponentType::Float ? 0 : 1);
+		PermutationVector.Set<FTmvStructuredBufferSwizzlePS::FElementFormat>(ComponentType == ETmvMediaFrameComponentType::Float ? 0 : 1);
 
 		// Can have tiles but be planar. Need another parameter for planar vs tiled layout.
 		ETmvSwizzleBufferLayouts BufferLayout = ETmvSwizzleBufferLayouts::PlanarFull;
@@ -308,17 +352,7 @@ namespace UE::TmvMediaConverter::Private
 				// Calculating the offset of the component from the start of the mip buffer.
 				SIZE_T ComponentOffsetInBytes = 0;
 				MipInfo.GetPlaneBufferOffset(PlaneIndex, ComponentOffsetInBytes);
-
-				if (PlaneInfo.ComponentLayout == ETmvMediaFrameComponentLayout::Interleaved)
-				{
-					// For interleaved components, add the interleaving offset.
-					ComponentOffsetInBytes += ComponentIndexInPlane * (PlaneInfo.Stride / PlaneInfo.NumComponents);
-				}
-				else
-				{
-					// For packed components, add the offset of the component.
-					ComponentOffsetInBytes += ComponentIndexInPlane * PlaneInfo.GetBytesPerComponent();
-				}
+				ComponentOffsetInBytes += PlaneInfo.GetStartComponentOffsetInBytes(ComponentIndexInPlane);
 				
 				// Note: Shader expects strides and offsets in shorts.
 				Parameters.TileBufferOffsets[ComponentIndex] = ComponentOffsetInBytes / 2;	// tiled mode, plane info is per tile.
@@ -330,58 +364,16 @@ namespace UE::TmvMediaConverter::Private
 			}
 		}
 
-		Parameters.AlphaScale = MipInfo.Planes.Num() == 4 ? GetDataNorm(MipInfo.Planes[3].Type, MipInfo.Planes[3].BitDepth) : 1.0f;
-
-		Parameters.bConvertYUV = MipInfo.ColorModel == ETmvMediaFrameColorModel::YUV ? 1u : 0u;
-		if (Parameters.bConvertYUV)
-		{
-			const FVector YUVOffset = GetYUVOffset(MipInfo.Planes[0].Type, MipInfo.Planes[0].BitDepth); 
-			const float DataNorm = GetDataNorm(MipInfo.Planes[0].Type, MipInfo.Planes[0].BitDepth);
-			Parameters.YUVMatrix = static_cast<FMatrix44f>(CombineColorTransformAndOffset(MediaShaders::YuvToRgbRec709Scaled, YUVOffset, DataNorm));
-		}
-		else
-		{
-			Parameters.YUVMatrix = FMatrix44f::Identity;
-		}
+		Parameters.ColorParams.AlphaScale = MipInfo.Planes.Num() == 4 ? GetDataNorm(MipInfo.Planes[3].Type, MipInfo.Planes[3].BitDepth) : 1.0f;
 		
+		SetupYuvColorParams(MipInfo, EDataRangeType::UInt, Parameters.ColorParams);
+		SetupColorParams(ConverterParams, MipInfo, Parameters.ColorParams);
+		SetupTintColor(TextureMipLevel, Parameters.ColorParams);
+	
 		if (ConverterParams.bHasTiles)
 		{
 			Parameters.NumTiles = FIntPoint(FMath::CeilToInt(float(SampleSize.X) / ConverterParams.TileDimWithBorders.X), FMath::CeilToInt(float(SampleSize.Y) / ConverterParams.TileDimWithBorders.Y));
 		}
-		if (ConverterParams.SourceColorSettings)
-		{
-			Parameters.bApplyColorTransform = 1u;
-			Parameters.EOTF = static_cast<uint32>(ConverterParams.SourceColorSettings->GetEncodingOverride());
-
-			const UE::Color::FColorSpace& DestinationCS = UE::Color::FColorSpace::GetWorking();
-			const UE::Color::FColorSpace& SourceCS = ConverterParams.SourceColorSettings->GetColorSpaceOverride(DestinationCS);
-			
-			if (SourceCS.Equals(DestinationCS))
-			{
-				Parameters.ColorSpaceMatrix = FMatrix44f::Identity;
-			}
-			else
-			{
-				const UE::Color::EChromaticAdaptationMethod Method = ConverterParams.SourceColorSettings->GetChromaticAdaptationMethod();
-				Parameters.ColorSpaceMatrix = UE::Color::Transpose<float>(UE::Color::FColorSpaceTransform(SourceCS, DestinationCS, Method));
-			}
-		}
-		else
-		{
-			Parameters.bApplyColorTransform = 0u;
-			Parameters.EOTF = 0u;
-			Parameters.ColorSpaceMatrix = FMatrix44f::Identity;
-		}
-
-		if (bTmvConverterTintMipmaps)
-		{
-			int TintIdx = FMath::Min(NumMipColors-1, TextureMipLevel);
-			Parameters.MipTint = MipmapColorizing[TintIdx];
-		}
-		else
-		{
-			Parameters.MipTint = FLinearColor(1.0f, 1.0f, 1.0f, 0.0f);
-		}
 		
 		if (ConverterParams.bHasTiles &&
 			(ConverterParams.TileInfoPerMipLevel.Num() > SampleMipLevel && ConverterParams.TileInfoPerMipLevel[SampleMipLevel].Num() > 0))
@@ -444,14 +436,13 @@ namespace UE::TmvMediaConverter::Private
 		// Using the same transition pattern as ConvertTextureToOutput.
 		RHICmdList.Transition(FRHITransitionInfo(RenderTargetTextureRHI, ERHIAccess::Unknown, ERHIAccess::RTV));
 		
-		// TMP: select the shader to use depending on having a unified mapped buffer or not. Need a better api for that.
-		if (BufferData->GetMappedBuffer())
+		if (BufferData->IsStructuredBuffer())
 		{
 			RenderMip_StructuredBuffer(RHICmdList, RenderTargetTextureRHI, ConverterParams, SampleMipLevel, TextureMipLevel, BufferData, SampleSize, TextureSize, MipViewports);
 		}
 		else
 		{
-			RenderMip_PlaneTextures(RHICmdList, RenderTargetTextureRHI, ConverterParams, SampleMipLevel, TextureMipLevel, BufferData, SampleSize, TextureSize, MipViewports);
+			RenderMip_PlaneTextures(RHICmdList, RenderTargetTextureRHI, ConverterParams, TextureMipLevel, BufferData, TextureSize, MipViewports);
 		}
 
 		RHICmdList.Transition(FRHITransitionInfo(RenderTargetTextureRHI, ERHIAccess::RTV, ERHIAccess::SRVGraphics));
@@ -461,7 +452,6 @@ namespace UE::TmvMediaConverter::Private
 	bool RenderThreadSwizzler(FRHICommandListImmediate& RHICmdList, FTextureRHIRef RenderTargetTextureRHI, TMap<int32, FTmvMediaFrameMipBufferHandle>& MipBuffers, const FTmvMediaTextureSampleConverterParameters ConverterParams)
 	{
 		RHI_BREADCRUMB_EVENT_STAT_F(RHICmdList, TmvMediaConverter, "TmvMediaConverter.Convert", "TmvMediaConverter.Convert %d", ConverterParams.FrameId);
-
 		int32 MipToUpscale = ConverterParams.UpscaleMip;
 
 		// Upscale to all mips below the mip to upscale.
@@ -501,7 +491,7 @@ namespace UE::TmvMediaConverter::Private
 			int MipLevelDiv = 1 << MipLevel;
 			FIntPoint Dim = ConverterParams.FullResolution / MipLevelDiv;
 
-			if (BufferData.IsValid() && BufferData->IsValid())
+			if (BufferData.IsValid() && BufferData->IsValidForRendering())
 			{
 				// Skip the mip to upscale because it is read and rendered already.
 				if (MipLevelViewports.Key == MipToUpscale)
@@ -522,7 +512,7 @@ FTmvMediaTextureSampleConverter::~FTmvMediaTextureSampleConverter() = default;
 bool FTmvMediaTextureSampleConverter::Convert(FRHICommandListImmediate& RHICmdList, FTextureRHIRef& InDstTexture, const FConversionHints& Hints)
 {
 	//smell: lock MipBufferCriticalSection to protect MipBuffers or make a copy?
-	return UE::TmvMediaConverter::Private::RenderThreadSwizzler(RHICmdList, InDstTexture, MipBuffers, GetParams_ThreadSafe());
+	return UE::TmvMedia::Converter::RenderThreadSwizzler(RHICmdList, InDstTexture, MipBuffers, GetParams_ThreadSafe());
 }
 
 bool FTmvMediaTextureSampleConverter::HasMipLevelBuffer(int32 RequestedMipLevel) const
diff --git a/Engine/Plugins/Media/TmvMedia/Source/TmvMedia/Private/TmvMediaModule.cpp b/Engine/Plugins/Media/TmvMedia/Source/TmvMedia/Private/TmvMediaModule.cpp
--- a/Engine/Plugins/Media/TmvMedia/Source/TmvMedia/Private/TmvMediaModule.cpp
+++ b/Engine/Plugins/Media/TmvMedia/Source/TmvMedia/Private/TmvMediaModule.cpp
@@ -1,18 +1,25 @@
 // Copyright Epic Games, Inc. All Rights Reserved.
 
-#include "ITmvMediaModule.h"
 #include "Decoder/ITmvMediaDecoderFactory.h"
+#include "Encoder/ITmvMediaEncoderFactory.h"
+#include "ITmvMediaModule.h"
 #include "Modules/ModuleManager.h"
 #include "TmvMediaLog.h"
+#include "Transcoder/TmvMediaTranscodeJobManager.h"
 
 DEFINE_LOG_CATEGORY(LogTmvMedia);
 
 class FTmvMediaModule : public ITmvMediaModule
 {
 	//~ Begin IModuleInterface
-	virtual bool SupportsDynamicReloading() override
+	virtual void StartupModule() override
+	{
+		TranscodeJobManager = MakeShared<FTmvMediaTranscodeJobManager>();
+	}
+
+	virtual void ShutdownModule() override
 	{
-		return false;
+		TranscodeJobManager.Reset();
 	}
 	//~ End IModuleInterface
 
@@ -22,6 +29,13 @@ class FTmvMediaModule : public ITmvMediaModule
 		if (InDecoderFactory.IsValid())
 		{
 			FScopeLock lock(&Lock);
+			if (DecoderFactoryNames.Contains(InDecoderFactory->GetName()))
+			{
+				UE_LOG(LogTmvMedia, Error, TEXT("Failed to add decoder factory \"%s\", it is already registered."), *InDecoderFactory->GetName());
+				return;
+			}
+
+			DecoderFactoryNames.Add(InDecoderFactory->GetName());
 			DecoderFactories.Emplace(MoveTemp(InDecoderFactory));
 		}
 	}
@@ -31,6 +45,7 @@ class FTmvMediaModule : public ITmvMediaModule
 		if (InDecoderFactory.IsValid())
 		{
 			FScopeLock lock(&Lock);
+			DecoderFactoryNames.Remove(InDecoderFactory->GetName());
 			DecoderFactories.RemoveSingleSwap(InDecoderFactory);
 		}
 	}
@@ -61,11 +76,68 @@ class FTmvMediaModule : public ITmvMediaModule
 		}
 		return BestFactory;
 	}
+
+	virtual void RegisterEncoderFactory(TSharedPtr<ITmvMediaEncoderFactory, ESPMode::ThreadSafe> InEncoderFactory) override
+	{
+		if (InEncoderFactory.IsValid())
+		{
+			FScopeLock lock(&Lock);
+			if (EncoderFactoryNames.Contains(InEncoderFactory->GetName()))
+			{
+				UE_LOG(LogTmvMedia, Error, TEXT("Failed to add encoder factory \"%s\", it is already registered."), *InEncoderFactory->GetName().ToString());
+				return;
+			}
+			EncoderFactoryNames.Add(InEncoderFactory->GetName());
+			EncoderFactories.Emplace(MoveTemp(InEncoderFactory));
+		}
+	}
+
+	virtual void UnregisterEncoderFactory(TSharedPtr<ITmvMediaEncoderFactory, ESPMode::ThreadSafe> InEncoderFactory) override
+	{
+		if (InEncoderFactory.IsValid())
+		{
+			FScopeLock lock(&Lock);
+			EncoderFactoryNames.Remove(InEncoderFactory->GetName());
+			EncoderFactories.RemoveSingleSwap(InEncoderFactory);
+		}
+	}
+
+	virtual void GetEncoderFactories(TArray<TWeakPtr<ITmvMediaEncoderFactory, ESPMode::ThreadSafe>>& OutEncoderFactories) const override
+	{
+		FScopeLock lock(&Lock);
+		OutEncoderFactories.Append(EncoderFactories);
+	}
+
+	virtual TSharedPtr<ITmvMediaEncoderFactory, ESPMode::ThreadSafe> FindEncoderFactory(FName InEncoderFactoryName) const override
+	{
+		FScopeLock lock(&Lock);
+		for (TWeakPtr<ITmvMediaEncoderFactory, ESPMode::ThreadSafe> EncoderFactoryWeak : EncoderFactories)
+		{
+			if (TSharedPtr<ITmvMediaEncoderFactory, ESPMode::ThreadSafe> EncoderFactory = EncoderFactoryWeak.Pin())
+			{
+				if (EncoderFactory->GetName() == InEncoderFactoryName)
+				{
+					return EncoderFactory;
+				}
+			}
+		}
+		return nullptr;
+	}
+
+	virtual ITmvMediaTranscodeJobManager* GetTranscodeJobManager() const override
+	{
+		return TranscodeJobManager.Get();
+	}
 	//~ End ITmvMediaModule
 
 private:
 	mutable FCriticalSection Lock;
 	TArray<TWeakPtr<ITmvMediaDecoderFactory, ESPMode::ThreadSafe>> DecoderFactories;
+	TSet<FString> DecoderFactoryNames;
+	TArray<TWeakPtr<ITmvMediaEncoderFactory, ESPMode::ThreadSafe>> EncoderFactories;
+	TSet<FName> EncoderFactoryNames;
+
+	TSharedPtr<FTmvMediaTranscodeJobManager> TranscodeJobManager;
 };
 
 ITmvMediaModule* ITmvMediaModule::Get()
@@ -78,5 +150,4 @@ ITmvMediaModule& ITmvMediaModule::GetOrLoad()
 	return FModuleManager::LoadModuleChecked<FTmvMediaModule>("TmvMedia");
 }
 
-
 IMPLEMENT_MODULE(FTmvMediaModule, TmvMedia);
\ No newline at end of file
diff --git a/Engine/Plugins/Media/TmvMedia/Source/TmvMedia/Private/Transcoder/TmvMediaFrameMips.cpp b/Engine/Plugins/Media/TmvMedia/Source/TmvMedia/Private/Transcoder/TmvMediaFrameMips.cpp
new file mode 100644
--- /dev/null
+++ b/Engine/Plugins/Media/TmvMedia/Source/TmvMedia/Private/Transcoder/TmvMediaFrameMips.cpp
@@ -0,0 +1,69 @@
+// Copyright Epic Games, Inc. All Rights Reserved.
+
+#include "Transcoder/TmvMediaFrameMips.h"
+
+#include "ImageCore.h"
+#include "SampleConverter/TmvMediaFrameMipImageBuffer.h"
+
+namespace UE::TmvMedia
+{
+	namespace TimeUtils
+	{
+		/** Convert a frame duration to a frame rate. */
+		FFrameRate GetFrameRate(const FTimespan& InFrameDuration, const FFrameRate& InDefault)
+		{
+			if (InFrameDuration > FTimespan::Zero())
+			{
+				const double Framerate = 1.0 / InFrameDuration.GetTotalSeconds();
+				// Fractional frame rates are considered "drop frame"
+				if (FMath::Fractional(Framerate) > UE_KINDA_SMALL_NUMBER)
+				{
+					return FFrameRate(FMath::CeilToInt(Framerate) * 1000, 1001);
+				}
+				return FFrameRate(FMath::CeilToInt(Framerate), 1);
+			}
+			return InDefault;
+		}
+	}
+}
+
+bool FTmvMediaFrameMips::Init(const TSharedPtr<FImage>& InMip0, bool bInGenerateMips)
+{
+	if (!InMip0)
+	{
+		return false;
+	}
+	
+	MipBuffers.Reset();
+	MipBuffers.Add(MakeShared<FTmvMediaFrameMipImageBuffer>(InMip0, 0));
+
+	// Directly use ImageCore to generate the mips on the cpu. (It's reasonably fast.)
+	if (bInGenerateMips && (InMip0->SizeX > 1 || InMip0->SizeY > 1))
+	{
+		TSharedPtr<FImage> PreviousMip = InMip0;
+		// Protecting against infinite loop.
+		constexpr int32 MaxMips = 32;
+
+		while (MipBuffers.Num() < MaxMips)
+		{
+			int32 MipSizeX = FMath::Max(PreviousMip->SizeX/2, 1);
+			int32 MipSizeY = FMath::Max(PreviousMip->SizeY/2, 1);
+			
+			TSharedPtr<FImage> NewMip = MakeShared<FImage>(MipSizeX, MipSizeY, 1, InMip0->Format, InMip0->GammaSpace);
+			{
+				TRACE_CPUPROFILER_EVENT_SCOPE(TmvMediaTranscode::GenerateMip);
+				FImageCore::ResizeImage(*PreviousMip, *NewMip);
+			}
+			int32 MipLevel = MipBuffers.Num();			
+			MipBuffers.Add(MakeShared<FTmvMediaFrameMipImageBuffer>(NewMip, MipLevel));
+
+			PreviousMip = NewMip;
+
+			if (MipSizeX == 1 && MipSizeY == 1)
+			{
+				break;
+			}
+		}
+	}
+	return true;
+}
diff --git a/Engine/Plugins/Media/TmvMedia/Source/TmvMedia/Private/Transcoder/TmvMediaTranscodeJob.cpp b/Engine/Plugins/Media/TmvMedia/Source/TmvMedia/Private/Transcoder/TmvMediaTranscodeJob.cpp
new file mode 100644
--- /dev/null
+++ b/Engine/Plugins/Media/TmvMedia/Source/TmvMedia/Private/Transcoder/TmvMediaTranscodeJob.cpp
@@ -0,0 +1,303 @@
+// Copyright Epic Games, Inc. All Rights Reserved.
+
+#include "Transcoder/TmvMediaTranscodeJob.h"
+
+#include "MediaSource.h"
+#include "Misc/Paths.h"
+#include "TmvMediaLog.h"
+#include "Transcoder/ITmvMediaTranscodeJobManager.h"
+#include "Transcoder/TmvMediaFrameProducer.h"
+#include "Utils/TmvMediaUtils.h"
+
+#include UE_INLINE_GENERATED_CPP_BY_NAME(TmvMediaTranscodeJob)
+
+#define LOCTEXT_NAMESPACE "TmvMediaTranscodeJob"
+
+namespace UE::TmvMedia::Transcode
+{
+	// Duplicated from: UImgMediaSource::ExpandSequencePathTokens.
+	// todo: Have some common "file" media source functions to centralize path handling.
+	FString ExpandSequencePathTokens(const FString& InPath)
+	{
+		return InPath
+			.Replace(TEXT("{engine_dir}"), *FPaths::ConvertRelativePathToFull(FPaths::EngineDir()))
+			.Replace(TEXT("{project_dir}"), *FPaths::ConvertRelativePathToFull(FPaths::ProjectDir()))
+			;
+	}
+	
+	FString ConvertPathToFull(const FString& InRelativeOrAbsolutePath)
+	{
+		// Expand the common known path tokens (potentially used in ImgMediaSource).
+		FString RelativeOrAbsolutePathExpanded = ExpandSequencePathTokens(InRelativeOrAbsolutePath);
+		
+		// If the path is already relative.
+		if (FPaths::IsRelative(RelativeOrAbsolutePathExpanded) || RelativeOrAbsolutePathExpanded.StartsWith("./"))
+		{
+			const TArray<FString> PossibleBasePaths =
+			{
+				FPaths::ProjectContentDir(),
+				FPaths::ProjectDir()
+			};
+
+			for (const FString& BasePath : PossibleBasePaths)
+			{
+				const FString BasePathFull = FPaths::ConvertRelativePathToFull(BasePath);
+				const FString PossiblePath = FPaths::ConvertRelativePathToFull(FPaths::Combine(BasePathFull, RelativeOrAbsolutePathExpanded));
+				if (FPaths::FileExists(PossiblePath))
+				{
+					return PossiblePath;
+				}
+			}
+		}
+
+		// Just try to make it absolute (in case it was still relative to binary).
+		return FPaths::ConvertRelativePathToFull(RelativeOrAbsolutePathExpanded);
+	}
+}
+
+
+FString FTmvMediaTranscodeJobSettings::GetInputPath() const
+{
+	if (InputSource == ETmvMediaTranscodeInputSource::File)
+	{
+		return InputPath.FilePath;
+	}
+
+	// Returns the input path.
+	return InputMediaSource.ToString();
+}
+
+FString FTmvMediaTranscodeJobSettings::GetAbsoluteInputPath() const
+{
+	if (InputSource == ETmvMediaTranscodeInputSource::File)
+	{
+		return UE::TmvMedia::Transcode::ConvertPathToFull(InputPath.FilePath);
+	}
+
+	// Get the media file from the media source.
+	if (UMediaSource* MediaSource = InputMediaSource.LoadSynchronous())
+	{
+		FString Url = MediaSource->GetUrl();
+
+		// Clean up to get just the filename.
+		const FString FilePrefix(TEXT("file://"));
+		const FString ImgPrefix(TEXT("img://"));
+
+		if (Url.StartsWith(FilePrefix))
+		{
+			return UE::TmvMedia::Transcode::ConvertPathToFull(Url.RightChop(FilePrefix.Len()));
+		}
+		
+		if (Url.StartsWith(ImgPrefix))
+		{
+			// Img Media only keeps the directory. For image sequence producer, it will then search for files in that directory.
+			return UE::TmvMedia::Transcode::ConvertPathToFull(Url.RightChop(ImgPrefix.Len()));
+		}
+
+		UE_LOG(LogTmvMedia, Error, TEXT("Media Source \"%s\" is not a supported codec for transcoding."), *Url);
+	}
+	else
+	{
+		UE_LOG(LogTmvMedia, Error, TEXT("Unable to load asset \"%s\"."), *InputMediaSource.ToString());
+	}
+
+	return FString();
+}
+
+
+bool FTmvMediaTranscodeJobSettings::IsInputPathSet() const
+{
+	if (InputSource == ETmvMediaTranscodeInputSource::File)
+	{
+		return InputPath.FilePath.IsEmpty() == false;
+	}
+
+	return InputMediaSource.IsNull() == false;
+}
+
+bool FTmvMediaTranscodeJobSettings::IsOutputPathSet() const
+{
+	return OutputPath.Path.IsEmpty() == false;
+}
+
+UTmvMediaTranscodeJob::UTmvMediaTranscodeJob()
+{
+	ITmvMediaTranscodeJobManager::SafeRegisterTranscodeJob(this);
+}
+
+void UTmvMediaTranscodeJob::SetId(const FGuid& InId)
+{
+	Id = InId;
+	ITmvMediaTranscodeJobManager::SafeRegisterTranscodeJob(this);
+}
+
+
+bool UTmvMediaTranscodeJob::Start(double InCurrentTime)
+{
+	UE_LOG(LogTmvMedia, Verbose, TEXT("Transcode Job for \"%s\" Started."), *Settings.GetInputPath());
+	
+	JobStatus = ETmvMediaTranscodeJobStatus::Running;
+	StopReason = ETmvMediaTranscodeJobStopReason::None;
+	bIsError = false;
+	JobStats = FTmvMediaTranscodingJobStats();
+	JobStats.StartTime = InCurrentTime;
+
+	UTmvMediaFrameProducer* FrameProducer = GetStage<UTmvMediaFrameProducer>();
+	if (!FrameProducer)
+	{
+		UE_LOG(LogTmvMedia, Error, TEXT("Transcode Job for \"%s\" failed to start: missing the Frame Producer stage."), *Settings.GetInputPath());
+		return false;
+	}
+
+	// Start all consumer stages first.
+	for (UTmvMediaTranscodeStage* Stage : Stages)
+	{
+		if (Stage != FrameProducer)
+		{
+			if (!Stage->Start(this))
+			{
+				JobStatus = ETmvMediaTranscodeJobStatus::Stopped;
+				OnStopped(InCurrentTime);
+				bIsError = true;
+				return false;
+			}
+		}
+	}
+
+	// Start the producer last.
+	if (!FrameProducer->Start(this))
+	{
+		JobStatus = ETmvMediaTranscodeJobStatus::Stopped;
+		OnStopped(InCurrentTime);
+		bIsError = true;
+	}
+
+	return IsRunning();
+}
+
+void UTmvMediaTranscodeJob::OnStopped(double InCurrentTime)
+{
+	JobStats.StopTime = InCurrentTime;
+
+	SafeBroadcastJobFinished(this);
+	
+	if (Notification.IsValid())
+	{
+		bool bIsCompleted;
+		{
+			FScopeLock lock(&JobStatsLock);
+			bIsCompleted = JobStats.ProcessedFrame == JobStats.TotalFramesToProcess;
+		}
+		Notification->Close(bIsCompleted);
+	}
+}
+
+void UTmvMediaTranscodeJob::RequestStop(double InCurrentTime, ETmvMediaTranscodeJobStopReason InReason)
+{
+	UE_LOG(LogTmvMedia, Verbose, TEXT("Transcode Job for \"%s\" Stopped. Reason: %s"),
+		*Settings.GetInputPath(), *UE::TmvMedia::Utils::StaticEnumToString(InReason)); 
+
+	const bool bWasRunning = IsRunning();
+	
+	if (IsRunning())
+	{
+		JobStatus = ETmvMediaTranscodeJobStatus::Stopping;
+		StopReason = InReason;
+	}
+	else
+	{
+		UE_LOG(LogTmvMedia, Verbose, TEXT("Transcode Job for \"%s\" was already stopping"), *Settings.GetInputPath()); 
+	}
+
+	// Stop the stages anyway.
+	for (UTmvMediaTranscodeStage* Stage : Stages)
+	{
+		Stage->RequestStop(this);	// Requests stop.
+	}
+
+	// Check if the stages are actually stopped.
+	if (AreStagesStopped())
+	{
+		JobStatus = ETmvMediaTranscodeJobStatus::Stopped;	// truly stopped.
+		OnStopped(InCurrentTime);
+	}
+}
+
+void UTmvMediaTranscodeJob::Discard(double InCurrentTime)
+{
+	if (IsRunning())
+	{
+		UE_LOG(LogTmvMedia, Warning, TEXT("Transcode Job for \"%s\" is being discarded while still running. Stopping... "), *Settings.GetInputPath());
+		RequestStop(InCurrentTime, ETmvMediaTranscodeJobStopReason::Discarded);
+	}
+	
+	UE_LOG(LogTmvMedia, Verbose, TEXT("Discarding Transcode Job for \"%s\"."), *Settings.GetInputPath()); 
+	ITmvMediaTranscodeJobManager::SafeUnregisterTranscodeJob(this);
+}
+
+void UTmvMediaTranscodeJob::Tick(const FTmvMediaTranscodeJobTime& InTime)
+{
+	if (JobStatus != ETmvMediaTranscodeJobStatus::Stopped)
+	{
+		for (UTmvMediaTranscodeStage* Stage : Stages)
+		{
+			Stage->Tick(this, InTime);
+		}
+
+		if (AreStagesStopped())
+		{
+			if (StopReason == ETmvMediaTranscodeJobStopReason::None)
+			{
+				StopReason = ETmvMediaTranscodeJobStopReason::Completed;
+			}
+			JobStatus = ETmvMediaTranscodeJobStatus::Stopped;
+			OnStopped(InTime.CurrentTime);
+		}
+	}
+}
+
+void UTmvMediaTranscodeJob::UpdateProgress(int32 InCurrentFrame, int32 InTotalFrames)
+{
+	if (Notification.IsValid())
+	{
+		Notification->SetText(FText::Format(LOCTEXT("TranscodeProgress", "Transcode Completed {0}/{1}"),
+			FText::AsNumber(InCurrentFrame), FText::AsNumber(InTotalFrames)));
+	}
+
+	{
+		FScopeLock lock(&JobStatsLock);
+		JobStats.ProcessedFrame = InCurrentFrame;
+		JobStats.TotalFramesToProcess = InTotalFrames;
+	}
+}
+
+FTmvMediaTranscodingJobStats UTmvMediaTranscodeJob::GetJobStats() const
+{
+	FScopeLock lock(&JobStatsLock);
+	return JobStats;
+}
+
+bool UTmvMediaTranscodeJob::IsCompleted() const
+{
+	return JobStatus == ETmvMediaTranscodeJobStatus::Stopped;
+}
+
+bool UTmvMediaTranscodeJob::AreStagesStopped() const 
+{
+	for (const UTmvMediaTranscodeStage* Stage : Stages)
+	{
+		if (Stage->GetStageStatus() != ETmvMediaTranscodeStageStatus::Stopped)
+		{
+			return false;
+		}
+	}
+	return true;
+}
+
+void UTmvMediaTranscodeJob::BeginDestroy()
+{
+	Super::BeginDestroy();
+	ITmvMediaTranscodeJobManager::SafeUnregisterTranscodeJob(this);
+}
+
+#undef LOCTEXT_NAMESPACE
\ No newline at end of file
diff --git a/Engine/Plugins/Media/TmvMedia/Source/TmvMedia/Private/Transcoder/TmvMediaTranscodeMuxer.cpp b/Engine/Plugins/Media/TmvMedia/Source/TmvMedia/Private/Transcoder/TmvMediaTranscodeMuxer.cpp
new file mode 100644
--- /dev/null
+++ b/Engine/Plugins/Media/TmvMedia/Source/TmvMedia/Private/Transcoder/TmvMediaTranscodeMuxer.cpp
@@ -0,0 +1,112 @@
+// Copyright Epic Games, Inc. All Rights Reserved.
+
+#include "Transcoder/TmvMediaTranscodeMuxer.h"
+
+#include "GenericPlatform/GenericPlatformFile.h"
+#include "HAL/PlatformFileManager.h"
+#include "Misc/Paths.h"
+#include "TmvMediaLog.h"
+#include "Transcoder/TmvMediaFrameMips.h"
+#include "Transcoder/TmvMediaFrameProducer.h"
+#include "Transcoder/TmvMediaTranscodeJob.h"
+
+#include UE_INLINE_GENERATED_CPP_BY_NAME(TmvMediaTranscodeMuxer)
+
+int32 UTmvMediaTranscodeMuxer::OpenStream(UTmvMediaTranscodeJob* InParentJob, const FString& InStreamName, const FString& InExtension)
+{
+	if (!InParentJob)
+	{
+		return INDEX_NONE;
+	}
+
+	if (InStreamName.IsEmpty() || InExtension.IsEmpty())
+	{
+		UE_LOG(LogTmvMedia, Error, TEXT("OpenStream: Invalid stream name or extension (Name=\"%s\", Ext=\"%s\")"), *InStreamName, *InExtension);
+		UTmvMediaTranscodeJob::SafeReportError(InParentJob);
+		return INDEX_NONE;
+	}
+
+	const FString OutPath = InParentJob->Settings.OutputPath.Path;
+
+	if (OutPath.IsEmpty())
+	{
+		UE_LOG(LogTmvMedia, Error, TEXT("OpenStream: OutputPath is empty."));
+		UTmvMediaTranscodeJob::SafeReportError(InParentJob);
+		return INDEX_NONE;
+	}
+
+	if (!FPaths::DirectoryExists(OutPath))
+	{
+		if (!FPlatformFileManager::Get().GetPlatformFile().CreateDirectoryTree(*OutPath))
+		{
+			UE_LOG(LogTmvMedia, Error, TEXT("OpenStream: Failed to create output directory \"%s\""), *OutPath);
+			UTmvMediaTranscodeJob::SafeReportError(InParentJob);
+			return INDEX_NONE;
+		}
+	}
+
+	FScopeLock Lock(&StreamsCS);
+	Streams.Add({InStreamName, InExtension});
+	return Streams.Num() - 1;
+}
+
+void UTmvMediaTranscodeMuxer::ReceiveAccessUnit(
+		UTmvMediaTranscodeJob* InParentJob,
+		int32 InStreamId,
+		const FTmvMediaFrameTimeInfo& InTimeInfo,
+		TSharedPtr<TArray64<uint8>>&& InAccessUnit)
+{
+	if (!InAccessUnit.IsValid())
+	{
+		UE_LOG(LogTmvMedia, Error, TEXT("ReceiveAccessUnit: Null access unit."));
+		UTmvMediaTranscodeJob::SafeReportError(InParentJob);
+		return; 
+	}
+
+	if (InAccessUnit->Num() == 0)
+	{
+		UE_LOG(LogTmvMedia, Warning, TEXT("ReceiveAccessUnit: Empty access unit (stream %d, frame %d)."), InStreamId, InTimeInfo.FrameIndex);
+		return;
+	}
+
+	FStreamInfo StreamInfo;
+	{
+		FScopeLock Lock(&StreamsCS);
+		if (!Streams.IsValidIndex(InStreamId))
+		{
+			UE_LOG(LogTmvMedia, Error, TEXT("ReceiveAccessUnit: Invalid stream id."));
+			UTmvMediaTranscodeJob::SafeReportError(InParentJob);
+			return;
+		}
+		StreamInfo = Streams[InStreamId];
+	}
+
+	if (!InParentJob)
+	{
+		UE_LOG(LogTmvMedia, Error, TEXT("ReceiveAccessUnit: Invalid Transcode Job."));
+		return;
+	}
+
+	// Remark: for a file sequence muxer, we don't need to synchronize (no common container) and
+	// we can just write the file immediately.
+	// For a container muxer, we need to keep the access units and re-order them. They can be
+	// received out of order because of multiple worker threads. We would need to know the number of worker threads.
+	const FString OutPath = InParentJob->Settings.OutputPath.Path;
+	const FString FileName = FString::Printf(TEXT("%s%05d.%s"), *StreamInfo.StreamName, InTimeInfo.FrameIndex, *StreamInfo.Extension);
+	const FString FullFileName = FPaths::Combine(OutPath, FPaths::MakeValidFileName(FileName));
+
+	TUniquePtr<IFileHandle> File(FPlatformFileManager::Get().GetPlatformFile().OpenWrite(*FullFileName));
+	if (File)
+	{
+		if (!File->Write(InAccessUnit->GetData(), InAccessUnit->Num()))
+		{
+			UE_LOG(LogTmvMedia, Error, TEXT("ReceiveAccessUnit: Failed to write access unit to file %s."), *FullFileName);
+			UTmvMediaTranscodeJob::SafeReportError(InParentJob);
+		}
+	}
+	else
+	{
+		UE_LOG(LogTmvMedia, Error, TEXT("ReceiveAccessUnit: Failed to open output file for writing: %s"), *FullFileName);
+		UTmvMediaTranscodeJob::SafeReportError(InParentJob);
+	}
+}
\ No newline at end of file
diff --git a/Engine/Plugins/Media/TmvMedia/Source/TmvMedia/Public/ITmvMediaModule.h b/Engine/Plugins/Media/TmvMedia/Source/TmvMedia/Public/ITmvMediaModule.h
--- a/Engine/Plugins/Media/TmvMedia/Source/TmvMedia/Public/ITmvMediaModule.h
+++ b/Engine/Plugins/Media/TmvMedia/Source/TmvMedia/Public/ITmvMediaModule.h
@@ -10,6 +10,8 @@
 #define UE_API TMVMEDIA_API
 
 class ITmvMediaDecoderFactory;
+class ITmvMediaEncoderFactory;
+class ITmvMediaTranscodeJobManager;
 
 /**
  * Interface for the Tiled-Mipmap Video (TMV) module.
@@ -49,6 +51,33 @@ public:
 	 * @return Highest priority decoder factory if found, null otherwise.
 	 */
 	virtual TSharedPtr<ITmvMediaDecoderFactory, ESPMode::ThreadSafe> GetBestDecoderFactoryForFormat(const FString& InCodecFormat, const TMap<FString, FVariant>& InOptions) = 0;
+
+	/**
+	 * Register an encoder factory in the module.
+	 */
+	virtual void RegisterEncoderFactory(TSharedPtr<ITmvMediaEncoderFactory, ESPMode::ThreadSafe> InEncoderFactory) = 0;
+
+	/**
+	 * Unregister an encoder factory from the module.
+	 */
+	virtual void UnregisterEncoderFactory(TSharedPtr<ITmvMediaEncoderFactory, ESPMode::ThreadSafe> InEncoderFactory) = 0;
+
+	/**
+	 * Get a list of all currently registered encoder factories. 
+	 */
+	virtual void GetEncoderFactories(TArray<TWeakPtr<ITmvMediaEncoderFactory, ESPMode::ThreadSafe>>& OutEncoderFactories) const = 0;
+
+	/**
+	 * Find the encoder factory by name.
+	 * @param InEncoderFactoryName Name of the factory.
+	 * @return The requested factory if found.
+	 */
+	virtual TSharedPtr<ITmvMediaEncoderFactory, ESPMode::ThreadSafe> FindEncoderFactory(FName InEncoderFactoryName) const = 0;
+
+	/**
+	 * Access the global transcode job manager.
+	 */
+	virtual ITmvMediaTranscodeJobManager* GetTranscodeJobManager() const = 0;
 };
 
-#undef UE_API
\ No newline at end of file
+#undef UE_API
diff --git a/Engine/Plugins/Media/TmvMedia/Source/TmvMedia/Public/SampleConverter/TmvMediaFrameMipBuffer.h b/Engine/Plugins/Media/TmvMedia/Source/TmvMedia/Public/SampleConverter/TmvMediaFrameMipBuffer.h
--- a/Engine/Plugins/Media/TmvMedia/Source/TmvMedia/Public/SampleConverter/TmvMediaFrameMipBuffer.h
+++ b/Engine/Plugins/Media/TmvMedia/Source/TmvMedia/Public/SampleConverter/TmvMediaFrameMipBuffer.h
@@ -74,8 +74,9 @@ public:
 
 	/**
 	 * Returns the start address of the memory plane buffer that contains the given component.
-	 * @remark In the case of interleaved components in a memory plane, an additional offset must be applied to the
-	 * returned value to get the start of the component within the plane buffer.
+	 * @remark In the case of interleaved components in a memory plane, an additional offset
+	 * provided by GetStartComponentOffsetInBytes(), must be applied to the returned value to
+	 * get the start of the component within the plane buffer.
 	 */
 	virtual void* GetPlaneBufferForComponent(int32 InComponentIndex) = 0;
 
@@ -99,7 +100,13 @@ public:
 	virtual void CopyTileRegions(int32 InFrameId, const FTmvMediaFrameMipInfo& InTmvMipInfo, TConstArrayView<FIntRect> InTileRegions, TConstArrayView<FTmvMediaShaderTileDesc> InTileInfos) = 0;
 
 	/** Returns true if the buffer can be used for rendering. */
-	virtual bool IsValid() const = 0;
+	virtual bool IsValidForRendering() const = 0;
+
+	/** Returns true if the buffer is implemented as a structured buffer. */
+	virtual bool IsStructuredBuffer() const
+	{
+		return false;
+	}
 
 protected:
 	/** Allocated buffer size, used for debugging. */
diff --git a/Engine/Plugins/Media/TmvMedia/Source/TmvMedia/Public/TmvMediaFrameInfo.h b/Engine/Plugins/Media/TmvMedia/Source/TmvMedia/Public/TmvMediaFrameInfo.h
--- a/Engine/Plugins/Media/TmvMedia/Source/TmvMedia/Public/TmvMediaFrameInfo.h
+++ b/Engine/Plugins/Media/TmvMedia/Source/TmvMedia/Public/TmvMediaFrameInfo.h
@@ -4,6 +4,8 @@
 
 #include "Containers/Array.h"
 #include "Math/IntPoint.h"
+#include "Math/UnrealMathUtility.h"
+#include "TmvMediaFrameColorInfo.h"
 
 #define UE_API TMVMEDIA_API
 
@@ -144,6 +146,31 @@ struct FTmvMediaFramePlaneInfo
 	{
 		return static_cast<SIZE_T>(Stride) * static_cast<SIZE_T>(NumLines);
 	}
+
+	/**
+	 * Computes the offset, in bytes, to the start of the given component in the plane buffer.
+	 * @param InComponentIndexInPlane Index of the given component in the current plane. Ex: in a UV plane, U would have index 0 and V, index 1.
+	 * @return Offset to apply to the start of the plane buffer to get the first pixel of the given component.
+	 */
+	uint32 GetStartComponentOffsetInBytes(int32 InComponentIndexInPlane) const
+	{
+		if (InComponentIndexInPlane <= 0 || NumComponents == 0) 
+		{ 
+			return 0; 
+		}
+
+		// Clamp to last valid component to avoid out-of-range offsets. 
+		const uint32 ClampedIndex = static_cast<uint32>(FMath::Min<int32>(InComponentIndexInPlane, static_cast<int32>(NumComponents) - 1)); 
+
+		if (ComponentLayout == ETmvMediaFrameComponentLayout::Interleaved)
+		{
+			// For interleaved components, add the interleaving offset.
+			return ClampedIndex * (Stride / NumComponents);
+		}
+		
+		// For packed components, add the offset of the component.
+		return ClampedIndex * GetBytesPerComponent();
+	}
 };
 
 /**
@@ -195,6 +222,9 @@ struct FTmvMediaFrameMipInfo
 	/** Indicate the color model that allow to determine how the component should be interpreted. */
 	ETmvMediaFrameColorModel ColorModel = ETmvMediaFrameColorModel::RGB;
 
+	/** Specifies the color information (color space, encoding, yuv matrix and range) for the given frame. */
+	FTmvMediaFrameColorInfo ColorInfo;
+
 	/** Number of horizontal samples in a tile. */
 	int32 TileWidth = 0;
 
diff --git a/Engine/Plugins/Media/TmvMedia/Source/TmvMedia/Public/Transcoder/TmvMediaFrameMips.h b/Engine/Plugins/Media/TmvMedia/Source/TmvMedia/Public/Transcoder/TmvMediaFrameMips.h
new file mode 100644
--- /dev/null
+++ b/Engine/Plugins/Media/TmvMedia/Source/TmvMedia/Public/Transcoder/TmvMediaFrameMips.h
@@ -0,0 +1,60 @@
+// Copyright Epic Games, Inc. All Rights Reserved.
+
+#pragma once
+
+#include "HAL/Platform.h"
+#include "Misc/FrameRate.h"
+#include "Misc/Timespan.h"
+#include "SampleConverter/TmvMediaFrameMipBuffer.h"
+#include "Templates/SharedPointer.h"
+
+#define UE_API TMVMEDIA_API
+
+struct FImage;
+
+/** Time information for a frame. */
+struct FTmvMediaFrameTimeInfo
+{
+	/** Frame index (for index based muxers). */
+	int32 FrameIndex = 0;
+
+	/** Frame time stamp (for time span based muxers). */
+	FTimespan FrameTime;
+
+	/** Frame time duration. */
+	FTimespan FrameDuration;
+};
+
+namespace UE::TmvMedia::TimeUtils
+{
+	/**
+	 * Calculate the frame rate (num, den) for the given frame duration.
+	 * This function will consider the fractional frame rate as "drop frame".
+	 * 
+	 * @param InFrameDuration frame duration
+	 * @param InDefault Default frame rate returns in case duration is invalid.
+	 * @return Calculated frame rate.
+	 */
+	UE_API FFrameRate GetFrameRate(const FTimespan& InFrameDuration, const FFrameRate& InDefault = FFrameRate());
+}
+
+/**
+ * Container of an array of frame mip buffers with time info for the transcoder pipeline.
+ */
+struct FTmvMediaFrameMips
+{
+	/** Time indexing information. */
+	FTmvMediaFrameTimeInfo TimeInfo;
+
+	/** Array of mips buffers. */
+	TArray<FTmvMediaFrameMipBufferHandle> MipBuffers;
+
+	/**
+	 * Helper to initialize the mip buffers from the provided mip0.
+	 * @param InMip0 Image to use for mip 0. Keeps a reference to the image.
+	 * @param bInGenerateMips if true, mips will be generated from mip 0.
+	 */
+	UE_API bool Init(const TSharedPtr<FImage>& InMip0, bool bInGenerateMips);
+};
+
+#undef UE_API
\ No newline at end of file
diff --git a/Engine/Plugins/Media/TmvMedia/Source/TmvMedia/TmvMedia.Build.cs b/Engine/Plugins/Media/TmvMedia/Source/TmvMedia/TmvMedia.Build.cs
--- a/Engine/Plugins/Media/TmvMedia/Source/TmvMedia/TmvMedia.Build.cs
+++ b/Engine/Plugins/Media/TmvMedia/Source/TmvMedia/TmvMedia.Build.cs
@@ -16,8 +16,14 @@ namespace UnrealBuildTool.Rules
 				new string[] {
 					"Core",
 					"CoreUObject",
+					"Engine",
+					"ImageCore",
+					"ImageWrapper",
+					"Json",
+					"JsonUtilities",
 					"Media",
 					"MediaAssets",
+					"MediaUtils",
 					"Projects", 
 					"RHI",
 					"RenderCore",
