[TmvMedia] TMV Transcoding Framework

diff --git a/Engine/Plugins/Media/TmvMedia/Source/TmvMediaEditor/Private/Widgets/STmvMediaTranscoderJobControls.cpp b/Engine/Plugins/Media/TmvMedia/Source/TmvMediaEditor/Private/Widgets/STmvMediaTranscoderJobControls.cpp
new file mode 100644
--- /dev/null
+++ b/Engine/Plugins/Media/TmvMedia/Source/TmvMediaEditor/Private/Widgets/STmvMediaTranscoderJobControls.cpp
@@ -0,0 +1,270 @@
+// Copyright Epic Games, Inc. All Rights Reserved.
+
+#include "STmvMediaTranscoderJobControls.h"
+
+#include "Editor.h"
+#include "Encoder/TmvMediaEncoderOptions.h"
+#include "IStructureDetailsView.h"
+#include "Modules/ModuleManager.h"
+#include "PropertyEditorModule.h"
+#include "SlateOptMacros.h"
+#include "TmvMediaEditorLog.h"
+#include "TmvMediaTranscodeListHandle.h"
+#include "TmvMediaTranscodeNotification.h"
+#include "Transcoder/TmvMediaFileSequenceFrameProducer.h"
+#include "Transcoder/TmvMediaFrameConverter.h"
+#include "Transcoder/TmvMediaPlayerFrameProducer.h"
+#include "Transcoder/TmvMediaTestFrameEncoder.h"
+#include "Transcoder/TmvMediaTmvFrameEncoder.h"
+#include "Transcoder/TmvMediaTranscodeJob.h"
+#include "Transcoder/TmvMediaTranscodeList.h"
+#include "Transcoder/TmvMediaTranscodeMuxer.h"
+#include "UObject/UObjectIterator.h"
+#include "Widgets/Input/SButton.h"
+#include "Widgets/Input/STextComboBox.h"
+#include "Widgets/Layout/SScrollBox.h"
+
+#define LOCTEXT_NAMESPACE "STmvMediaTranscoderJobControls"
+
+STmvMediaTranscoderJobControls::~STmvMediaTranscoderJobControls()
+{
+	CleanUp();
+}
+
+BEGIN_SLATE_FUNCTION_BUILD_OPTIMIZATION
+void STmvMediaTranscoderJobControls::Construct(const FArguments& InArgs, const TSharedPtr<FTmvMediaTranscodeListHandle>& InListHandle)
+{
+	ListHandle = InListHandle;
+	
+	ChildSlot
+	[
+		// Add process images button.
+		SNew(SHorizontalBox)
+		+ SHorizontalBox::Slot()
+			.AutoWidth()
+			[
+				SNew(SButton)
+					.ButtonStyle(&FAppStyle::Get().GetWidgetStyle<FButtonStyle>("PrimaryButton"))
+					.IsEnabled(this, &STmvMediaTranscoderJobControls::CanStartJobs)
+					.OnClicked(this, &STmvMediaTranscoderJobControls::OnStartJobs)
+					.ForegroundColor(FSlateColor::UseStyle())
+					.Text(LOCTEXT("StartJobButtonLabel", "Start Job(s)"))
+					.ToolTipText(LOCTEXT("StartJobButtonToolTip", "Start transcoding selected jobs."))
+			]
+
+		+ SHorizontalBox::Slot()
+			.AutoWidth()
+			[
+				SNew(SButton)
+					.IsEnabled(this, &STmvMediaTranscoderJobControls::CanCancelJobs)
+					.OnClicked(this, &STmvMediaTranscoderJobControls::OnCancelJobs)
+					.Text(LOCTEXT("CancelJobButtonLabel", "Cancel"))
+					.ToolTipText(LOCTEXT("CancelJobButtonToolTip", "Cancel all transcoding jobs."))
+			]
+	];
+}
+END_SLATE_FUNCTION_BUILD_OPTIMIZATION
+
+namespace UE::TmvMediaEditor::Transcode
+{
+	void CleanupJob(UTmvMediaTranscodeJob* InJob, const double InCurrentTime)
+	{
+		InJob->RequestStop(InCurrentTime, ETmvMediaTranscodeJobStopReason::None);
+		InJob->Discard(InCurrentTime);
+	}
+	
+	void CleanupJob(UTmvMediaTranscodeJob* InJob)
+	{
+		CleanupJob(InJob, FApp::GetCurrentTime());
+	}
+	
+	// todo: move this in a transcode job builder in the runtime module.
+	UTmvMediaTranscodeJob* MakeJob(const FTmvMediaTranscodeListItem* InJobItem)
+	{
+		UTmvMediaTranscodeJob* TranscodeJob = NewObject<UTmvMediaTranscodeJob>(GetTransientPackage(), NAME_None);
+		TranscodeJob->SetId(InJobItem->Id);
+		TranscodeJob->Settings = InJobItem->Settings;
+		
+		if (TranscodeJob->Settings.bUseMediaPlayer)
+		{
+			TranscodeJob->AddStage(NewObject<UTmvMediaPlayerFrameProducer>(TranscodeJob, NAME_None));
+		}
+		else
+		{
+			TranscodeJob->AddStage(NewObject<UTmvMediaFileSequenceFrameProducer>(TranscodeJob, NAME_None));
+		}
+		
+		TranscodeJob->AddStage(NewObject<UTmvMediaFrameConverter>(TranscodeJob, NAME_None));
+		TranscodeJob->AddStage(NewObject<UTmvMediaTmvFrameEncoder>(TranscodeJob, NAME_None));
+		//TranscodeJob->AddStage(NewObject<UTmvMediaTestFrameEncoder>(TranscodeJob, NAME_None));
+		TranscodeJob->AddStage(NewObject<UTmvMediaTranscodeMuxer>(TranscodeJob, NAME_None));
+		TranscodeJob->GetStage<UTmvMediaFrameEncoder>()->SetEncoderOptions(InJobItem->EncoderOptions);
+		
+		return TranscodeJob;
+	}
+}
+
+void STmvMediaTranscoderJobControls::Tick(const FGeometry& AllottedGeometry, const double InCurrentTime, const float InDeltaTime)
+{
+	SCompoundWidget::Tick(AllottedGeometry, InCurrentTime, InDeltaTime);
+
+	if (CurrentTranscodeJob)
+	{
+		CurrentTranscodeJob->Tick({InCurrentTime, InDeltaTime});
+		
+		if (CurrentTranscodeJob->IsCompleted())
+		{
+			UE::TmvMediaEditor::Transcode::CleanupJob(CurrentTranscodeJob.Get(), InCurrentTime);
+			CurrentTranscodeJob.Reset();
+			
+			// Start next pending job.
+			if (!PendingTranscodeJobs.IsEmpty())
+			{
+				CurrentTranscodeJob = PendingTranscodeJobs.Pop();
+				if (CurrentTranscodeJob)
+				{
+					CurrentTranscodeJob->SetNotificationHandler(MakeShared<FTmvMediaTranscodeNotificationSafe>());
+					CurrentTranscodeJob->Start(InCurrentTime);
+				}
+			}
+		}
+	}
+}
+
+bool STmvMediaTranscoderJobControls::IsProcessing() const
+{
+	return CurrentTranscodeJob.IsValid() && CurrentTranscodeJob->IsRunning();
+}
+
+bool STmvMediaTranscoderJobControls::IsJobProcessingOrPending(const FGuid& InJobId) const
+{
+	if (CurrentTranscodeJob && CurrentTranscodeJob->GetId() == InJobId)
+	{
+		return true;
+	}
+
+	for (const TStrongObjectPtr<UTmvMediaTranscodeJob>& PendingJob : PendingTranscodeJobs)
+	{
+		if (PendingJob && PendingJob->GetId() == InJobId)
+		{
+			return true;
+		}
+	}
+
+	return false;
+}
+
+bool STmvMediaTranscoderJobControls::CanStartJobItem(const FTmvMediaTranscodeListItem& InJobItem) const
+{
+	return InJobItem.Settings.IsInputPathSet() && InJobItem.Settings.IsOutputPathSet() && !IsJobProcessingOrPending(InJobItem.Id);
+}
+
+bool STmvMediaTranscoderJobControls::CanStartJobs() const
+{
+	UTmvMediaTranscodeList* List = ListHandle ? ListHandle->Get() : nullptr;
+	if (!List)
+	{
+		return false;
+	}
+
+	bool bHasStartableJobs = false;
+
+	// Check that we don't have any items already started or that is not configured to run.
+	for (const int32 SelectedItem : ListHandle->GetCurrentSelection())
+	{
+		if (List->IsValidItemIndex(SelectedItem))
+		{
+			if (!CanStartJobItem(List->GetItem(SelectedItem)))
+			{
+				return false;
+			}
+
+			bHasStartableJobs = true;
+		}
+	}
+
+	return bHasStartableJobs;
+}
+
+FReply STmvMediaTranscoderJobControls::OnStartJobs()
+{
+	UTmvMediaTranscodeList* List = ListHandle ? ListHandle->Get() : nullptr;
+	if (!List)
+	{
+		return FReply::Handled();
+	}
+
+	TArray<int32> SelectedItems = ListHandle->GetCurrentSelection();
+	for (int32 SelectedItem : SelectedItems)
+	{
+		if (List->IsValidItemIndex(SelectedItem))
+		{
+			const FTmvMediaTranscodeListItem& JobItem = List->GetItem(SelectedItem);
+		
+			if (!CanStartJobItem(JobItem))
+			{
+				continue;
+			}
+			
+			UTmvMediaTranscodeJob* NewJob = UE::TmvMediaEditor::Transcode::MakeJob(&JobItem);
+			if (!CurrentTranscodeJob)
+			{
+				CurrentTranscodeJob.Reset(NewJob);
+				CurrentTranscodeJob->SetNotificationHandler(MakeShared<FTmvMediaTranscodeNotificationSafe>());
+				CurrentTranscodeJob->Start(FApp::GetCurrentTime());
+			}
+			else
+			{
+				// Since we only run one job at a time, the other jobs will go to the pending list.
+				PendingTranscodeJobs.Add(TStrongObjectPtr<UTmvMediaTranscodeJob>(NewJob));
+			}
+		}
+	}
+
+	return FReply::Handled();
+}
+
+bool STmvMediaTranscoderJobControls::CanCancelJobs() const
+{
+	return IsProcessing();
+}
+
+FReply STmvMediaTranscoderJobControls::OnCancelJobs()
+{
+	const double CurrentTime = FApp::GetCurrentTime();
+	
+	if (CurrentTranscodeJob.IsValid())
+	{
+		// Stop will return immediately, but the task is still stopping. 
+		CurrentTranscodeJob->RequestStop(CurrentTime, ETmvMediaTranscodeJobStopReason::Cancelled);
+	}
+
+	// Forget all pending jobs too.
+	for (const TStrongObjectPtr<UTmvMediaTranscodeJob>& PendingJob : PendingTranscodeJobs )
+	{
+		PendingJob->Discard(CurrentTime);
+	}
+	PendingTranscodeJobs.Reset();
+
+	return FReply::Handled();
+}
+
+void STmvMediaTranscoderJobControls::CleanUp()
+{
+	const double CurrentTime = FApp::GetCurrentTime();
+
+	if (CurrentTranscodeJob.IsValid())
+	{
+		UE::TmvMediaEditor::Transcode::CleanupJob(CurrentTranscodeJob.Get(), CurrentTime);
+		CurrentTranscodeJob.Reset();
+	}
+	
+	// Also discard all pending jobs.
+	for (const auto& PendingJob : PendingTranscodeJobs )
+	{
+		PendingJob->Discard(CurrentTime);
+	}
+	PendingTranscodeJobs.Reset();
+}
+
+#undef LOCTEXT_NAMESPACE
diff --git a/Engine/Plugins/Media/TmvMedia/Source/TmvMediaEditor/Private/Widgets/STmvMediaTranscoderJobControls.h b/Engine/Plugins/Media/TmvMedia/Source/TmvMediaEditor/Private/Widgets/STmvMediaTranscoderJobControls.h
new file mode 100644
--- /dev/null
+++ b/Engine/Plugins/Media/TmvMedia/Source/TmvMediaEditor/Private/Widgets/STmvMediaTranscoderJobControls.h
@@ -0,0 +1,63 @@
+// Copyright Epic Games, Inc. All Rights Reserved.
+
+#pragma once
+
+#include "Input/Reply.h"
+#include "Transcoder/TmvMediaTranscodeJob.h"
+#include "UObject/StrongObjectPtr.h"
+#include "Widgets/SCompoundWidget.h"
+
+struct FTmvMediaTranscodeListHandle;
+struct FTmvMediaTranscodeListItem;
+
+/**
+ * Tmv Transcoder job controls widget
+ */
+class STmvMediaTranscoderJobControls : public SCompoundWidget
+{
+public:
+	SLATE_BEGIN_ARGS(STmvMediaTranscoderJobControls){}
+	SLATE_END_ARGS()
+
+	virtual ~STmvMediaTranscoderJobControls() override;
+
+	void Construct(const FArguments& InArgs, const TSharedPtr<FTmvMediaTranscodeListHandle>& InListHandle);
+	
+	//~ Begin SWidget
+	virtual void Tick(const FGeometry& AllottedGeometry, const double InCurrentTime, const float InDeltaTime) override;
+	//~ End SWidget
+
+	/** Returns true if at least one job is processing. */
+	bool IsProcessing() const;
+
+private:
+	/** Returns true if the job is either processing or pending processing. */
+	bool IsJobProcessingOrPending(const FGuid& InJobId) const;
+
+	/** Stops and discards the current job and discards all pending jobs. */
+	void CleanUp();
+
+	/** Returns true if the given job item can be started. */
+	bool CanStartJobItem(const FTmvMediaTranscodeListItem& InJobItem) const;
+	
+	/** Returns true if the currently selected set of jobs can be started. */
+	bool CanStartJobs() const;
+
+	/** Called when we click on the start jobs button. */
+	FReply OnStartJobs();
+
+	/** Returns true if the current set of running or pending jobs can be cancelled. */
+	bool CanCancelJobs() const;
+
+	/** Called when we click on the cancel button. */
+	FReply OnCancelJobs();
+
+	/** Reference to the transcode list */
+	TSharedPtr<FTmvMediaTranscodeListHandle> ListHandle;
+
+	/** Currently executing transcode job. We only run one job at a time. */
+	TStrongObjectPtr<UTmvMediaTranscodeJob> CurrentTranscodeJob;
+
+	/** List of pending transcode jobs (not yet executing). */
+	TArray<TStrongObjectPtr<UTmvMediaTranscodeJob>> PendingTranscodeJobs;
+};
diff --git a/Engine/Plugins/Media/TmvMedia/Source/TmvMediaEditor/Private/Widgets/STmvMediaTranscoderJobDetails.cpp b/Engine/Plugins/Media/TmvMedia/Source/TmvMediaEditor/Private/Widgets/STmvMediaTranscoderJobDetails.cpp
new file mode 100644
--- /dev/null
+++ b/Engine/Plugins/Media/TmvMedia/Source/TmvMediaEditor/Private/Widgets/STmvMediaTranscoderJobDetails.cpp
@@ -0,0 +1,309 @@
+// Copyright Epic Games, Inc. All Rights Reserved.
+
+#include "STmvMediaTranscoderJobDetails.h"
+
+#include "Customizations/TmvMediaFilePathCustomization.h"
+#include "Editor.h"
+#include "Encoder/TmvMediaEncoderOptions.h"
+#include "IStructureDetailsView.h"
+#include "Modules/ModuleManager.h"
+#include "PropertyEditorModule.h"
+#include "SlateOptMacros.h"
+#include "TmvMediaEditorLog.h"
+#include "TmvMediaEditorTranscodeUtils.h"
+#include "TmvMediaTranscodeListHandle.h"
+#include "Transcoder/ITmvMediaTranscodeJobManager.h"
+#include "Transcoder/TmvMediaTranscodeJob.h"
+#include "Transcoder/TmvMediaTranscodeList.h"
+#include "Widgets/Input/STextComboBox.h"
+#include "Widgets/Layout/SBox.h"
+
+#define LOCTEXT_NAMESPACE "STmvMediaTranscoderJobDetails"
+
+STmvMediaTranscoderJobDetails::~STmvMediaTranscoderJobDetails()
+{
+	if (ListHandle)
+	{
+		if (UTmvMediaTranscodeList* List = ListHandle->Get())
+		{
+			List->GetOnItemEvent().RemoveAll(this);
+		}
+		ListHandle->GetOnTranscodeListChanged().RemoveAll(this);
+		ListHandle->GetOnSelectionChanged().RemoveAll(this);
+	}
+}
+
+BEGIN_SLATE_FUNCTION_BUILD_OPTIMIZATION
+void STmvMediaTranscoderJobDetails::Construct(const FArguments& InArgs, const TSharedPtr<FTmvMediaTranscodeListHandle>& InListHandle)
+{
+	ListHandle = InListHandle;
+	SelectedJobItemIndex = INDEX_NONE; // nothing selected on construction.
+
+	if (ListHandle)
+	{
+		if (UTmvMediaTranscodeList* List = ListHandle->Get())
+		{
+			List->GetOnItemEvent().AddSP(this, &STmvMediaTranscoderJobDetails::OnTranscodeListItemEvent);
+		}
+		ListHandle->GetOnSelectionChanged().AddSP(this, &STmvMediaTranscoderJobDetails::OnJobItemSelectionChanged);
+		ListHandle->GetOnTranscodeListChanged().AddSP(this, &STmvMediaTranscoderJobDetails::OnTranscodeListChanged);
+	}
+	
+	// Set up widgets.
+	TSharedPtr<SBox> JobDetailsViewBox;
+	TSharedPtr<SBox> EncoderDetailsViewBox;
+
+	SetupEncoderOptionsSelection();
+	
+	ChildSlot
+	[
+		SNew(SVerticalBox)
+		.IsEnabled(this, &STmvMediaTranscoderJobDetails::IsDetailsEnabled)
+
+		// Add details view.
+		+ SVerticalBox::Slot()
+		.AutoHeight()
+			[
+				SAssignNew(JobDetailsViewBox, SBox)
+			]
+		+ SVerticalBox::Slot()
+		.AutoHeight()
+			[
+				// Encoder selector
+				SNew(SHorizontalBox)
+				+SHorizontalBox::Slot()
+				.AutoWidth()
+				.Padding(0.0f, 3.0f, 0.0f, 0.0f)
+				[
+					SNew(STextBlock)
+					.Text(LOCTEXT("TmvMediaEncoderSelectionLabel", "Encoder"))
+				]
+				+ SHorizontalBox::Slot()
+				.AutoWidth()
+				.Padding(5.0f, 0.0f, 0.0f, 0.0f)
+				[
+					SAssignNew(EncoderOptionsComboBox, STextComboBox)
+					.OptionsSource(&EncoderOptionsNames)
+					.InitiallySelectedItem(GetCurrentEncoderOptionsItem())
+					.OnSelectionChanged(this, &STmvMediaTranscoderJobDetails::OnEncoderOptionsSelectionChanged)
+				]
+			]
+		+ SVerticalBox::Slot()
+		.AutoHeight()
+			[
+				SAssignNew(EncoderDetailsViewBox, SBox)
+			]
+	];
+
+	RefreshJobSettingsDetailsView();
+	if (JobDetailsView && JobDetailsViewBox)
+	{
+		JobDetailsViewBox->SetContent(JobDetailsView->GetWidget().ToSharedRef());
+	}
+	
+	RefreshEncoderOptionsDetailsView();
+	if (EncoderDetailsView && EncoderDetailsViewBox)
+	{
+		EncoderDetailsViewBox->SetContent(EncoderDetailsView->GetWidget().ToSharedRef());
+	}
+}
+END_SLATE_FUNCTION_BUILD_OPTIMIZATION
+
+bool STmvMediaTranscoderJobDetails::IsDetailsEnabled() const
+{
+	// The whole details panel is disabled if the job is either queued or running.
+	UTmvMediaTranscodeList* List = ListHandle ? ListHandle->Get() : nullptr;
+	if (List && SelectedJobItemIndex != INDEX_NONE)
+	{
+		const auto& Item = List->GetItem(SelectedJobItemIndex);
+		if (const ITmvMediaTranscodeJobManager* TranscodeJobManager = ITmvMediaTranscodeJobManager::Get())
+		{
+			if (TranscodeJobManager->GetTranscodeJob(Item.Id))
+			{
+				// If a job is already queued, it won't be edited anymore (settings are copied in the queued job).
+				return false;
+			}
+		}
+	}
+	return true;
+}
+
+void STmvMediaTranscoderJobDetails::OnJobItemSelectionChanged(const UTmvMediaTranscodeList* InList, TConstArrayView<int32> InSelectedItems)
+{
+	SelectedJobItemIndex = !InSelectedItems.IsEmpty() ? InSelectedItems[0] : INDEX_NONE;
+	RefreshJobSettingsDetailsView();
+	RefreshEncoderOptionsDetailsView();
+	RefreshEncoderOptionsCombo();
+}
+
+void STmvMediaTranscoderJobDetails::OnTranscodeListItemEvent(const UTmvMediaTranscodeList* InList, const FTmvMediaTranscodeListItemEventArgs& InArgs)
+{
+	if (InArgs.Type == ETmvMediaTranscodeListItemEventType::ItemsModified && bIgnoreModifiedItemEvents)
+	{
+		return;
+	}
+	
+	// Note: when items are added, it may reallocate the array and the memory pointed to by
+	// the detail views will become dangling. Views must be refreshed.
+	RefreshJobSettingsDetailsView();
+	RefreshEncoderOptionsDetailsView();
+	RefreshEncoderOptionsCombo();
+}
+
+void STmvMediaTranscoderJobDetails::OnTranscodeListChanged(UTmvMediaTranscodeList* InPreviousList, UTmvMediaTranscodeList* InNewList)
+{
+	if (InPreviousList)
+	{
+		InPreviousList->GetOnItemEvent().RemoveAll(this);
+	}
+	if (InNewList)
+	{
+		InNewList->GetOnItemEvent().AddSP(this, &STmvMediaTranscoderJobDetails::OnTranscodeListItemEvent);
+	}
+
+	RefreshJobSettingsDetailsView();
+	RefreshEncoderOptionsDetailsView();
+	RefreshEncoderOptionsCombo();
+}
+
+void STmvMediaTranscoderJobDetails::OnFinishedChangingProperties(const FPropertyChangedEvent& InChangedEvent)
+{
+	UTmvMediaTranscodeList* List = ListHandle ? ListHandle->Get() : nullptr;
+	if (List && SelectedJobItemIndex != INDEX_NONE)
+	{
+		TGuardValue TickGuard(bIgnoreModifiedItemEvents, true);	// Ignore events coming from self.
+		List->GetOnItemEvent().Broadcast(List, {ETmvMediaTranscodeListItemEventType::ItemsModified, {SelectedJobItemIndex}});
+	}
+}
+
+FTmvMediaTranscodeListItem* STmvMediaTranscoderJobDetails::GetSelectedJobItem() const
+{
+	UTmvMediaTranscodeList* List = ListHandle ? ListHandle->Get() : nullptr;
+	if (SelectedJobItemIndex == INDEX_NONE || !List)
+	{
+		return nullptr;
+	}
+	return List->GetItemMutable(SelectedJobItemIndex);
+}
+
+void STmvMediaTranscoderJobDetails::RefreshJobSettingsDetailsView()
+{
+	FTmvMediaTranscodeListItem* JobItem = GetSelectedJobItem();
+	TSharedPtr<FStructOnScope> JobOptionsWrap = JobItem ? 
+		MakeShared<FStructOnScope>(FTmvMediaTranscodeJobSettings::StaticStruct(), (uint8*)&JobItem->Settings) : MakeShared<FStructOnScope>();
+	
+	if (!JobDetailsView)
+	{
+		JobDetailsView = UE::TmvMediaEditor::Transcode::CreateStructureDetailView(JobOptionsWrap);
+		if (JobDetailsView && JobDetailsView->GetDetailsView())
+		{
+			// Add our filepath customization that supports drag and drop and path resolving.
+			JobDetailsView->GetDetailsView()->RegisterInstancedCustomPropertyTypeLayout(FName(TEXT("FilePath")),
+				FOnGetPropertyTypeCustomizationInstance::CreateStatic(FTmvMediaFilePathCustomization::MakeInstance));
+			
+			JobDetailsView->GetDetailsView()->OnFinishedChangingProperties().AddSP(this, &STmvMediaTranscoderJobDetails::OnFinishedChangingProperties);
+		}
+
+		if (!JobDetailsView)
+		{
+			UE_LOG(LogTmvMediaEditor, Error, TEXT("Failed to create job options detail view."));
+		}
+	}
+	else
+	{
+		JobDetailsView->SetStructureData(JobOptionsWrap);
+	}
+}
+
+void STmvMediaTranscoderJobDetails::RefreshEncoderOptionsDetailsView()
+{
+	FTmvMediaTranscodeListItem* JobItem = GetSelectedJobItem();
+	TSharedPtr<FStructOnScope> EncoderOptionsWrap = JobItem && JobItem->EncoderOptions.IsValid() ? 
+		MakeShared<FStructOnScope>(JobItem->EncoderOptions.GetScriptStruct(), JobItem->EncoderOptions.GetMutableMemory()) : MakeShared<FStructOnScope>();
+	
+	if (!EncoderDetailsView)
+	{
+		EncoderDetailsView = UE::TmvMediaEditor::Transcode::CreateStructureDetailView(EncoderOptionsWrap);
+		if (!EncoderDetailsView)
+		{
+			UE_LOG(LogTmvMediaEditor, Error, TEXT("Failed to create encoder options detail view."));
+		}
+	}
+	else
+	{
+		EncoderDetailsView->SetStructureData(EncoderOptionsWrap);
+	}
+}
+
+void STmvMediaTranscoderJobDetails::RefreshEncoderOptionsCombo()
+{
+	FTmvMediaTranscodeListItem* JobItem = GetSelectedJobItem();
+	
+	const FString SelectedEncoderName = JobItem && JobItem->EncoderOptions.IsValid() ? JobItem->EncoderOptions.GetScriptStruct()->GetDisplayNameText().ToString() : TEXT("None");
+	
+	const int32 NewSelectedEncoderIndex = EncoderOptionsNames.IndexOfByPredicate([&SelectedEncoderName](const TSharedPtr<FString>& InEncoderName)
+	{
+		return *InEncoderName == SelectedEncoderName;
+	});
+
+	if (NewSelectedEncoderIndex != INDEX_NONE)
+	{
+		EncoderOptionsIndex = NewSelectedEncoderIndex;
+		if (EncoderOptionsComboBox)
+		{
+			EncoderOptionsComboBox->SetSelectedItem(GetCurrentEncoderOptionsItem());
+		}
+	}
+	else
+	{
+		UE_LOG(LogTmvMediaEditor, Error, TEXT("Could not find encoder option struct type: %s"), *SelectedEncoderName);
+	}
+}
+
+void STmvMediaTranscoderJobDetails::SetupEncoderOptionsSelection()
+{
+	EncoderOptionsStructs = UE::TmvMediaEditor::Transcode::GetAllDerivedStruct(FTmvMediaEncoderOptions::StaticStruct());
+	EncoderOptionsNames.Reset(EncoderOptionsStructs.Num() + 1);
+
+	for (const UScriptStruct* Struct : EncoderOptionsStructs)
+	{
+		EncoderOptionsNames.Add(MakeShared<FString>(Struct->GetDisplayNameText().ToString()));
+	}
+
+	// Initialize current selection index to the first element.
+	EncoderOptionsIndex = 0;
+}
+
+void STmvMediaTranscoderJobDetails::OnEncoderOptionsSelectionChanged(TSharedPtr<FString> InNewSelection, ESelectInfo::Type InSelectInfo)
+{
+	FTmvMediaTranscodeListItem* SelectedJobItem = GetSelectedJobItem();
+	if (!SelectedJobItem)
+	{
+		UE_LOG(LogTmvMediaEditor, Error, TEXT("No Job item selected. Select a job item in the list first."));
+		return;
+	}
+	
+	int32 NewSelectionIndex;
+	if (EncoderOptionsNames.Find(InNewSelection, NewSelectionIndex))
+	{
+		if (EncoderOptionsIndex != NewSelectionIndex)
+		{
+			EncoderOptionsIndex = NewSelectionIndex;
+
+			if (EncoderDetailsView)
+			{
+				EncoderDetailsView->SetStructureData(MakeShared<FStructOnScope>());	// Clear to avoid dangling.
+			}
+		
+			const UScriptStruct* EncoderOptionsStruct = GetCurrentEncoderOptionsStruct();
+			if (EncoderOptionsStruct && SelectedJobItem->EncoderOptions.GetScriptStruct() != EncoderOptionsStruct)
+			{
+				SelectedJobItem->EncoderOptions.InitializeAsScriptStruct(EncoderOptionsStruct);
+			}
+
+			RefreshEncoderOptionsDetailsView();
+		}
+	}
+}
+
+#undef LOCTEXT_NAMESPACE
diff --git a/Engine/Plugins/Media/TmvMedia/Source/TmvMediaEditor/Private/Widgets/STmvMediaTranscoderJobDetails.h b/Engine/Plugins/Media/TmvMedia/Source/TmvMediaEditor/Private/Widgets/STmvMediaTranscoderJobDetails.h
new file mode 100644
--- /dev/null
+++ b/Engine/Plugins/Media/TmvMedia/Source/TmvMediaEditor/Private/Widgets/STmvMediaTranscoderJobDetails.h
@@ -0,0 +1,115 @@
+// Copyright Epic Games, Inc. All Rights Reserved.
+
+#pragma once
+
+#include "Widgets/SCompoundWidget.h"
+
+class IStructureDetailsView;
+class STextComboBox;
+class UTmvMediaTranscodeList;
+struct FTmvMediaTranscodeListHandle;
+struct FTmvMediaTranscodeListItem;
+struct FTmvMediaTranscodeListItemEventArgs;
+
+/**
+ * Transcoder Job Settings and Encoder Options Detail views for the Tmv transcoder list editor. 
+ */
+class STmvMediaTranscoderJobDetails : public SCompoundWidget
+{
+public:
+	SLATE_BEGIN_ARGS(STmvMediaTranscoderJobDetails){}
+	SLATE_END_ARGS()
+
+	virtual ~STmvMediaTranscoderJobDetails() override;
+	
+	void Construct(const FArguments& InArgs, const TSharedPtr<FTmvMediaTranscodeListHandle>& InListHandle);
+
+private:
+	/** Returns true if the detail panel should be enabled for editing. */
+	bool IsDetailsEnabled() const;
+
+	/** Handles transcode list selection changed to update the details view. */
+	void OnJobItemSelectionChanged(const UTmvMediaTranscodeList* InList, TConstArrayView<int32> InSelectedItems);
+
+	/** Handles transcode list event (add, remove, modified, etc) to update the details view. */
+	void OnTranscodeListItemEvent(const UTmvMediaTranscodeList* InList, const FTmvMediaTranscodeListItemEventArgs& InArgs);
+
+	/** Handles transcode list changed (i.e. a new list is loaded) to update the details view. */
+	void OnTranscodeListChanged(UTmvMediaTranscodeList* InPreviousList, UTmvMediaTranscodeList* InNewList);
+
+	/** Details View event handling. Propagates the detail view event. */
+	void OnFinishedChangingProperties(const FPropertyChangedEvent& InChangedEvent);
+
+	/** Get currently "selected" job item, i.e. item that is reflected in the detail views. */
+	FTmvMediaTranscodeListItem* GetSelectedJobItem() const;
+
+	/** Refresh (create or update) the Job Setting Details View. */
+	void RefreshJobSettingsDetailsView();
+
+	/** Refresh (create or update) the Encoder Options Details View. */
+	void RefreshEncoderOptionsDetailsView();
+
+	/** Encoder Options Combo Box - Make sure the encoder option combo box is reflecting the current job item. */
+	void RefreshEncoderOptionsCombo();
+
+	/** Encoder Options Combo Box - Populate Encoder Options Structs list for Combo Box. */
+	void SetupEncoderOptionsSelection();
+
+	/** Encoder Options Combo Box -Returns the currently selected Encoder Options struct. */ 
+	UScriptStruct* GetCurrentEncoderOptionsStruct()
+	{
+		if (EncoderOptionsStructs.IsEmpty())
+		{
+			SetupEncoderOptionsSelection();
+		}
+		if (EncoderOptionsStructs.IsValidIndex(EncoderOptionsIndex))
+		{
+			return EncoderOptionsStructs[EncoderOptionsIndex];
+		}
+		return nullptr;
+	}
+
+	/** Encoder Options Combo Box - Handles updating the job item and details view when selection change. */  
+	void OnEncoderOptionsSelectionChanged(TSharedPtr<FString> InNewSelection, ESelectInfo::Type InSelectInfo);
+
+	/** Encoder Options Combo Box - Returns the selected "item" from the backing list. */
+	TSharedPtr<FString> GetCurrentEncoderOptionsItem()
+	{
+		if (EncoderOptionsStructs.IsEmpty())
+		{
+			SetupEncoderOptionsSelection();
+		}
+		if (EncoderOptionsNames.IsValidIndex(EncoderOptionsIndex))
+		{
+			return EncoderOptionsNames[EncoderOptionsIndex];
+		}
+		return nullptr;
+	}
+
+	/** Handle to the currently opened Transcode List. */
+	TSharedPtr<FTmvMediaTranscodeListHandle> ListHandle;
+
+	/** Currently selected Job Item Index (in the job list) that is reflected by this detail view. */
+	int32 SelectedJobItemIndex = INDEX_NONE;
+
+	/** Flag indicating the modified items events are coming from the detail panel itself and should be ignored. */
+	bool bIgnoreModifiedItemEvents = false;
+	
+	/** Holds the "Job Settings" details view. */
+	TSharedPtr<IStructureDetailsView> JobDetailsView;
+	
+	/** Holds the "Encoder Options" details view. */
+	TSharedPtr<IStructureDetailsView> EncoderDetailsView;
+
+	/** Reference to Encoder Options ComboBox widget to propagate "refresh" events. */
+	TSharedPtr<STextComboBox> EncoderOptionsComboBox;
+
+	/** Cached list of the child ScriptStruct of EncoderOptions struct to the encoder options combo box. */
+	TArray<UScriptStruct*> EncoderOptionsStructs;
+
+	/** Display String items (struct display names) for the encoder options combo box. */
+	TArray<TSharedPtr<FString>> EncoderOptionsNames;
+
+	/** Currently selected encoder option's struct for the given job item. */
+	int32 EncoderOptionsIndex = 0;
+};
diff --git a/Engine/Plugins/Media/TmvMedia/Source/TmvMediaEditor/Private/Widgets/TmvMediaEditorTranscodeUtils.cpp b/Engine/Plugins/Media/TmvMedia/Source/TmvMediaEditor/Private/Widgets/TmvMediaEditorTranscodeUtils.cpp
new file mode 100644
--- /dev/null
+++ b/Engine/Plugins/Media/TmvMedia/Source/TmvMediaEditor/Private/Widgets/TmvMediaEditorTranscodeUtils.cpp
@@ -0,0 +1,52 @@
+// Copyright Epic Games, Inc. All Rights Reserved.
+
+#include "TmvMediaEditorTranscodeUtils.h"
+
+#include "IStructureDetailsView.h"
+#include "Modules/ModuleManager.h"
+#include "PropertyEditorModule.h"
+#include "UObject/StructOnScope.h"
+#include "UObject/UObjectIterator.h"
+
+namespace UE::TmvMediaEditor::Transcode
+{
+	TArray<UScriptStruct*> GetAllDerivedStruct(const UScriptStruct* InBaseStruct)
+	{
+		TArray<UScriptStruct*> OutStructs;
+
+		if (InBaseStruct == nullptr)
+		{
+			return OutStructs;
+		}
+
+		OutStructs.Reserve(2);
+		
+		for (TObjectIterator<UScriptStruct> It; It; ++It)
+		{
+			UScriptStruct* CurrentStruct = *It;
+			if (CurrentStruct != InBaseStruct && CurrentStruct->IsChildOf(InBaseStruct))
+			{
+				OutStructs.Add(CurrentStruct);	
+			}
+		}
+		return OutStructs;
+	}
+
+	TSharedPtr<IStructureDetailsView> CreateStructureDetailView(const TSharedPtr<FStructOnScope>& InStructOnScope)
+	{
+		FStructureDetailsViewArgs StructureViewArgs;
+		StructureViewArgs.bShowObjects = true;
+		StructureViewArgs.bShowAssets = true;
+		StructureViewArgs.bShowClasses = true;
+		StructureViewArgs.bShowInterfaces = true;
+
+		FDetailsViewArgs ViewArgs;
+		ViewArgs.bAllowSearch = false;
+		ViewArgs.bHideSelectionTip = false;
+		ViewArgs.bShowObjectLabel = false;
+		ViewArgs.NameAreaSettings = FDetailsViewArgs::HideNameArea;
+
+		FPropertyEditorModule& PropertyEditor = FModuleManager::GetModuleChecked<FPropertyEditorModule>("PropertyEditor");
+		return PropertyEditor.CreateStructureDetailView(ViewArgs, StructureViewArgs, InStructOnScope);
+	}
+}
\ No newline at end of file
diff --git a/Engine/Plugins/Media/TmvMedia/Source/TmvMediaEditor/Private/Widgets/TmvMediaEditorTranscodeUtils.h b/Engine/Plugins/Media/TmvMedia/Source/TmvMediaEditor/Private/Widgets/TmvMediaEditorTranscodeUtils.h
new file mode 100644
--- /dev/null
+++ b/Engine/Plugins/Media/TmvMedia/Source/TmvMediaEditor/Private/Widgets/TmvMediaEditorTranscodeUtils.h
@@ -0,0 +1,19 @@
+// Copyright Epic Games, Inc. All Rights Reserved.
+
+#pragma once
+
+#include "Containers/Array.h"
+#include "Templates/SharedPointer.h"
+
+class FStructOnScope;
+class IStructureDetailsView;
+class UScriptStruct;
+
+namespace UE::TmvMediaEditor::Transcode
+{
+	/** Enumerate all the derived struct. */
+	TArray<UScriptStruct*> GetAllDerivedStruct(const UScriptStruct* InBaseStruct);
+	
+	/** Helper function to create a struct detail view for transcode job details panel. */
+	TSharedPtr<IStructureDetailsView> CreateStructureDetailView(const TSharedPtr<FStructOnScope>& InStructOnScope);
+}
diff --git a/Engine/Plugins/Media/TmvMedia/Source/TmvMediaShaders/Internal/TmvStructuredBufferSwizzlingShader.h b/Engine/Plugins/Media/TmvMedia/Source/TmvMediaShaders/Internal/TmvStructuredBufferSwizzlingShader.h
--- a/Engine/Plugins/Media/TmvMedia/Source/TmvMediaShaders/Internal/TmvStructuredBufferSwizzlingShader.h
+++ b/Engine/Plugins/Media/TmvMedia/Source/TmvMediaShaders/Internal/TmvStructuredBufferSwizzlingShader.h
@@ -6,6 +6,7 @@
 #include "ShaderParameterUtils.h"
 #include "ShaderParameterStruct.h"
 #include "TmvMediaShaderDefines.h"
+#include "TmvMediaShaderParameters.h"
 
 // The vertex shader used by DrawScreenPass to draw a rectangle.
 class FTmvSwizzleVS : public FGlobalShader
@@ -64,12 +65,6 @@ class FTmvStructuredBufferSwizzlePS : public FGlobalShader
 		SHADER_PARAMETER(FIntVector4, PlaneBufferStrides)
 		SHADER_PARAMETER(FIntVector4, PlaneWidthRatios)
 		SHADER_PARAMETER(FIntVector4, PlaneHeightRatios)
-		SHADER_PARAMETER(uint32, bApplyColorTransform)
-		SHADER_PARAMETER(uint32, EOTF)
-		SHADER_PARAMETER(FMatrix44f, ColorSpaceMatrix)
-		SHADER_PARAMETER(uint32, bConvertYUV)
-		SHADER_PARAMETER(FMatrix44f, YUVMatrix)
-		SHADER_PARAMETER(float, AlphaScale)
-		SHADER_PARAMETER(FVector4f, MipTint)
+		SHADER_PARAMETER_STRUCT_INCLUDE(FTmvMediaShaderColorParameters, ColorParams)
 	END_SHADER_PARAMETER_STRUCT()
 };
\ No newline at end of file
diff --git a/Engine/Plugins/Media/TmvMedia/Source/TmvMediaShaders/Internal/TmvTextureBufferSwizzlingShader.h b/Engine/Plugins/Media/TmvMedia/Source/TmvMediaShaders/Internal/TmvTextureBufferSwizzlingShader.h
--- a/Engine/Plugins/Media/TmvMedia/Source/TmvMediaShaders/Internal/TmvTextureBufferSwizzlingShader.h
+++ b/Engine/Plugins/Media/TmvMedia/Source/TmvMediaShaders/Internal/TmvTextureBufferSwizzlingShader.h
@@ -5,6 +5,7 @@
 #include "GlobalShader.h"
 #include "ShaderParameterUtils.h"
 #include "ShaderParameterStruct.h"
+#include "TmvMediaShaderParameters.h"
 
 /**
  * Pixel shader swizzle multi-component planar texture buffer into a packed RGBA texture.
@@ -24,12 +25,6 @@ class FTmvTextureBufferSwizzlePS : public FGlobalShader
 		SHADER_PARAMETER_SRV(Texture2D, InPlaneTexture2)
 		SHADER_PARAMETER_SRV(Texture2D, InPlaneTexture3)
 		SHADER_PARAMETER_SAMPLER(SamplerState, InPlaneTextureSampler)
-		SHADER_PARAMETER(uint32, bApplyColorTransform)
-		SHADER_PARAMETER(uint32, EOTF)
-		SHADER_PARAMETER(FMatrix44f, ColorSpaceMatrix)
-		SHADER_PARAMETER(uint32, bConvertYUV)
-		SHADER_PARAMETER(FMatrix44f, YUVMatrix)
-		SHADER_PARAMETER(float, AlphaScale)
-		SHADER_PARAMETER(FVector4f, MipTint)
+		SHADER_PARAMETER_STRUCT_INCLUDE(FTmvMediaShaderColorParameters, ColorParams)
 	END_SHADER_PARAMETER_STRUCT()
 };
\ No newline at end of file
diff --git a/Engine/Plugins/Media/TmvMedia/TmvMedia.uplugin b/Engine/Plugins/Media/TmvMedia/TmvMedia.uplugin
--- a/Engine/Plugins/Media/TmvMedia/TmvMedia.uplugin
+++ b/Engine/Plugins/Media/TmvMedia/TmvMedia.uplugin
@@ -21,6 +21,11 @@
 			"Type" : "Runtime",
 			"LoadingPhase" : "PostEngineInit"
 		},
+		{
+			"Name" : "TmvMediaEditor",
+			"Type" : "Editor",
+			"LoadingPhase" : "PostEngineInit"
+		},
 		{
 			"Name" : "TmvMediaShaders",
 			"Type" : "Runtime",
