[TmvMedia] Minor Fix:  CDO Transcode Job doesn't register to the Job Manager.

This would mainly lead to having a parasite active job when using live coding.

#Synced-CL 49891360

diff --git a/Engine/Plugins/Media/TmvMedia/Source/TmvMedia/Private/Transcoder/TmvMediaTranscodeJob.cpp b/Engine/Plugins/Media/TmvMedia/Source/TmvMedia/Private/Transcoder/TmvMediaTranscodeJob.cpp
--- a/Engine/Plugins/Media/TmvMedia/Source/TmvMedia/Private/Transcoder/TmvMediaTranscodeJob.cpp
+++ b/Engine/Plugins/Media/TmvMedia/Source/TmvMedia/Private/Transcoder/TmvMediaTranscodeJob.cpp
@@ -122,7 +122,10 @@ bool FTmvMediaTranscodeJobSettings::IsOutputPathSet() const
 
 UTmvMediaTranscodeJob::UTmvMediaTranscodeJob()
 {
-	ITmvMediaTranscodeJobManager::SafeRegisterTranscodeJob(this);
+	if(!HasAnyFlags(RF_ClassDefaultObject))
+	{
+		ITmvMediaTranscodeJobManager::SafeRegisterTranscodeJob(this);
+	}
 }
 
 void UTmvMediaTranscodeJob::SetId(const FGuid& InId)
@@ -131,7 +134,6 @@ void UTmvMediaTranscodeJob::SetId(const FGuid& InId)
 	ITmvMediaTranscodeJobManager::SafeRegisterTranscodeJob(this);
 }
 
-
 bool UTmvMediaTranscodeJob::Start(double InCurrentTime)
 {
 	UE_LOG(LogTmvMedia, Verbose, TEXT("Transcode Job for \"%s\" Started."), *Settings.GetInputPath());
