[TmvMedia] Transcoding Framework Improvements

1- Transcoding Job Settings and UI: Reworking the paths to be more compatible with the Img/File Media Source path schemes.

- Implementing much more elaborate path sanitization which works well either inside or outside of the project directory. Supports non existing paths (output folder), tokens. Added unit tests.
- Implemented a new directory picker widget that looks and feels more like the file picker. This allows us to customize the picked path sanitization.
- Centralized the path sanitization functions and made a pass to use it where applicable and remove code duplication.
- Adding a DirectoryPath customization to be able to sanitize the output directory like we do for the input directory.
- Implemented copy (Ctrl-C), paste (Ctrl-V), duplicate (Ctrl-D), delete (Del) and rename (F2) job items key bindings.
- implemented job list reordering (drag and drop)

2- Pass to improve the notification (toaster) error messages

- Fixed green mark to not show up if there is an error.
- Added an error message to SafeReportError. Start failures will display a message in the toaster.

3- Fixed support of incomplete sets of pre-calc mips (from existing exr sequences)

- When using the MediaPlayer to generate frames, the encoded mips are going to match what is generated by the player. Sequences with incomplete sets, like QAVirtualProduction\EXT_Assets\Movies\EXR_Sequences will only convert the 3 pre-calculated mips of that sequence.

#Synced-CL 49891360

diff --git a/Engine/Plugins/Media/TmvMedia/Source/TmvMedia/Internal/Transcoder/TmvMediaTranscodeSerialization.h b/Engine/Plugins/Media/TmvMedia/Source/TmvMedia/Internal/Transcoder/TmvMediaTranscodeSerialization.h
--- a/Engine/Plugins/Media/TmvMedia/Source/TmvMedia/Internal/Transcoder/TmvMediaTranscodeSerialization.h
+++ b/Engine/Plugins/Media/TmvMedia/Source/TmvMedia/Internal/Transcoder/TmvMediaTranscodeSerialization.h
@@ -11,6 +11,7 @@ class FArchive;
 class UTmvMediaTranscodeList;
 struct FTmvMediaEncoderOptions;
 struct FTmvMediaTranscodeJobSettings;
+struct FTmvMediaTranscodeListItem;
 
 namespace UE::TmvMedia::TranscodeSerialization
 {
@@ -35,6 +36,16 @@ namespace UE::TmvMedia::TranscodeSerialization
 	UE_API bool DeserializeTranscodeListFromJson(
 		FArchive& InArchive,
 		UTmvMediaTranscodeList& OutTranscodeList);
+
+	/** Serialize the given transcode list items to the given archive in json format. */
+	UE_API bool SerializeTranscodeListItemsToJson(
+		FArchive& InArchive,
+		const TArray<FTmvMediaTranscodeListItem>& InTranscodeListItems);
+
+	/** Deserialize the given transcode list items from the given archive in json format. */
+	UE_API bool DeserializeTranscodeListItemsFromJson(
+		FArchive& InArchive,
+		TArray<FTmvMediaTranscodeListItem>& OutTranscodeListItems);
 }
 
 #undef UE_API
\ No newline at end of file
diff --git a/Engine/Plugins/Media/TmvMedia/Source/TmvMedia/Internal/Utils/TmvMediaPathUtils.h b/Engine/Plugins/Media/TmvMedia/Source/TmvMedia/Internal/Utils/TmvMediaPathUtils.h
new file mode 100644
--- /dev/null
+++ b/Engine/Plugins/Media/TmvMedia/Source/TmvMedia/Internal/Utils/TmvMediaPathUtils.h
@@ -0,0 +1,52 @@
+// Copyright Epic Games, Inc. All Rights Reserved.
+
+#pragma once
+
+#include "Containers/UnrealString.h"
+
+#define UE_API TMVMEDIA_API
+
+/**
+ * Path utility functions.
+ * We want the TmvMedia to work with the paths from File/Img Media Sources.
+ */
+namespace UE::TmvMedia::PathUtils
+{
+	/**
+	 * Utility function to ensure the path is starting with a "./".
+	 * This is the format of the img/file media source use for internal relative paths.
+	 */
+	UE_API FString EnsureStartWithDotSlash(const FString& InPath);
+
+	/**
+	 * Utility function to extract the full directory path from a given sanitized file or directory path.
+	 * This function will work reliably for existing paths, otherwise, it will fall back to FPaths::GetPath.
+	 * @remark Relative paths are assumed to be relative to the project directory, unless using explicit tokens.
+	 * 
+	 * @param InSanitizedPath Can be a file or directory path in the sanitized format.
+	 */
+	UE_API FString GetDirectoryFullPath(const FString& InSanitizedPath);
+
+	/**
+	 * Utility function to convert a given path (file or directory) to a sanitized version
+	 * according to the tmv media supported paths format:
+	 * - relative path under project folder (will start with ./)
+	 * - absolute path if outside of project folder.
+	 * - support path tokens, i.e. supported tokens will remain.
+	 *
+	 * @param InPath Input path, can be a file or folder.
+	 * @return Sanitized path (file or folder).
+	 */
+	UE_API FString GetSanitizedPath(const FString& InPath);
+
+	/**
+	 * Utility function to convert the given path (file or directory, relative or absolute) to
+	 * a full path. This is intended to be the inverse operation of GetSanitizedPath().
+	 * 
+	 * @param InSanitizedPath Path to convert to absolute, can be a directory or file, relative or absolute. Also support path tokens.
+	 * @return Absolute path
+	 */
+	UE_API FString ConvertSanitizedPathToFull(const FString& InSanitizedPath);
+}
+
+#undef UE_API
\ No newline at end of file
diff --git a/Engine/Plugins/Media/TmvMedia/Source/TmvMedia/Internal/Utils/TmvMediaSerializationUtils.h b/Engine/Plugins/Media/TmvMedia/Source/TmvMedia/Internal/Utils/TmvMediaSerializationUtils.h
--- a/Engine/Plugins/Media/TmvMedia/Source/TmvMedia/Internal/Utils/TmvMediaSerializationUtils.h
+++ b/Engine/Plugins/Media/TmvMedia/Source/TmvMedia/Internal/Utils/TmvMediaSerializationUtils.h
@@ -2,6 +2,8 @@
 
 #pragma once
 
+#include "Containers/ArrayView.h"
+#include "Containers/UnrealString.h"
 #include "Templates/SharedPointer.h"
 
 #define UE_API TMVMEDIA_API
@@ -25,6 +27,16 @@ namespace UE::TmvMedia::SerializationUtils
 	 * Utility function to serialize a struct object to json which supports InstancedStruct fields. 
 	 */
 	UE_API TSharedPtr<FJsonObject> SerializeToJson(const UStruct* InStruct, const void* InObject);
+
+	/** Converts a serialized value stored as raw bytes into a string. */
+	UE_API FString BytesToString(const TArray<uint8>& InValueAsBytes);
+	
+	/** This ArrayView version is meant to work with FMemoryReaderView. */
+	inline TArrayView<const uint8> ValueToConstBytesView(const FString& InValueAsString)
+	{
+		// Len() excludes the terminating character, and it is desired so it matches what FJsonStructSerializerBackend does.
+		return TArrayView<const uint8>(reinterpret_cast<const uint8*>(*InValueAsString), InValueAsString.Len() * sizeof(TCHAR));
+	}
 }
 
 #undef UE_API
\ No newline at end of file
diff --git a/Engine/Plugins/Media/TmvMedia/Source/TmvMedia/Private/Tests/TmvMediaPathUtilsTests.cpp b/Engine/Plugins/Media/TmvMedia/Source/TmvMedia/Private/Tests/TmvMediaPathUtilsTests.cpp
new file mode 100644
--- /dev/null
+++ b/Engine/Plugins/Media/TmvMedia/Source/TmvMedia/Private/Tests/TmvMediaPathUtilsTests.cpp
@@ -0,0 +1,398 @@
+// Copyright Epic Games, Inc. All Rights Reserved.
+
+#include "Containers/UnrealString.h"
+#include "Engine/Engine.h"
+#include "HAL/FileManager.h"
+#include "Logging/LogMacros.h"
+#include "Misc/AutomationTest.h"
+#include "Misc/Paths.h"
+#include "Utils/TmvMediaPathUtils.h"
+
+#if WITH_DEV_AUTOMATION_TESTS
+
+namespace UE::TmvMedia::Tests
+{
+	/** Test Utility to create temporary directories and delete them when done. */
+	struct FTestFilesAndDirectories
+	{
+		~FTestFilesAndDirectories()
+		{
+			for (const FString& File : CreatedFiles)
+			{
+				IFileManager::Get().Delete(*File);
+			}
+
+			for(const FString& Directory : ReverseIterate(CreatedDirectories))
+			{
+				IFileManager::Get().DeleteDirectory(*Directory);
+			}
+		}
+
+		bool MakeDirectory(const FString& InPath)
+		{
+			if (FPaths::DirectoryExists(InPath))
+			{
+				return true;
+			}
+			if (IFileManager::Get().MakeDirectory(*InPath))
+			{
+				CreatedDirectories.Add(InPath);
+				return true;
+			}
+			return false;
+		}
+
+		bool MakeFile(const FString& InFilepath)
+		{
+			if (FPaths::FileExists(InFilepath))
+			{
+				return true;
+			}
+			const TUniquePtr<FArchive> FileWriter(IFileManager::Get().CreateFileWriter(*InFilepath));
+			if (FileWriter)
+			{
+				int32 value = 1234;
+				(*FileWriter) << value;
+				CreatedFiles.Add(InFilepath);
+				return true;
+			}
+			return false;
+		}
+
+
+	private:
+		TArray<FString> CreatedDirectories;
+		TArray<FString> CreatedFiles;
+	};
+}
+
+IMPLEMENT_SIMPLE_AUTOMATION_TEST(FTmvMediaSanitizePathTests, "System.Plugins.TmvMedia.SanitizePath", EAutomationTestFlags::EditorContext | EAutomationTestFlags::EngineFilter)
+bool FTmvMediaSanitizePathTests::RunTest(const FString& Parameters)
+{
+	using namespace UE::TmvMedia::Tests;
+	
+	// The test needs some existing files in different folders.
+	FTestFilesAndDirectories TestFilesAndDirectories;
+
+	auto MakeTempTestFileAndPath = [&TestFilesAndDirectories](const FString& InBasePath)
+	{
+		const FString MoviesDirectory = FPaths::ConvertRelativePathToFull(FPaths::Combine(InBasePath, TEXT("Movies")));
+		const FString SequenceDirectory = FPaths::Combine(MoviesDirectory, TEXT("TmvMediaTestSequence"));
+		const FString SequenceImgPath = FPaths::Combine(SequenceDirectory, TEXT("Image0000.exr"));
+		if (TestFilesAndDirectories.MakeDirectory(MoviesDirectory))
+		{
+			if (TestFilesAndDirectories.MakeDirectory(SequenceDirectory))
+			{
+				return TestFilesAndDirectories.MakeFile(SequenceImgPath);
+			}
+		}
+		return false;
+	};
+
+	// Prepare temporary files under different base paths. We want the Content vs Project one to test priority.
+	bool bSetupSuccess = true;
+	bSetupSuccess &= MakeTempTestFileAndPath(FPaths::ProjectContentDir());
+	bSetupSuccess &= MakeTempTestFileAndPath(FPaths::ProjectDir());
+	bSetupSuccess &= MakeTempTestFileAndPath(FPaths::EngineDir());
+
+	// If the setup failed, issue a warning to help troubleshooting CIS.
+	if (!bSetupSuccess)
+	{
+		AddWarning(TEXT("Failed to create the test media files and folders."));
+	}
+
+	// ---- Tests Begin here
+
+	// Testing GetSanitizedPath
+	auto TestSanitizePath = [this](const FString& InPath, const FString& InExpectedResult, FString& OutError)
+	{
+		const FString SanitizedPath = UE::TmvMedia::PathUtils::GetSanitizedPath(InPath);
+		if (SanitizedPath != InExpectedResult)
+		{
+			OutError = FString::Printf(TEXT("GetSanitizedPath failed: Input: \"%s\" Result: \"%s\" Expected \"%s\""), *InPath, *SanitizedPath, *InExpectedResult);
+			return false;
+		}
+		return true;
+	};
+
+	// Testing ConvertPathToFull, this is the complement to GetSanitizedPath.
+	auto TestConvertPathToFull = [this](const FString& InSanitizedPath, const FString& InExpectedResult, FString& OutError)
+	{
+		const FString FullPath = UE::TmvMedia::PathUtils::ConvertSanitizedPathToFull(InSanitizedPath);
+		if (FullPath != InExpectedResult)
+		{
+			OutError = FString::Printf(TEXT("ConvertPathToFull failed: Input: \"%s\" Result: \"%s\" Expected: \"%s\""), *InSanitizedPath, *FullPath, *InExpectedResult);
+			return false;
+		}
+		return true;
+	};
+
+	FString Error;
+	
+	// Empty path
+	AddErrorIfFalse(TestSanitizePath(TEXT(""), TEXT(""), Error), Error);
+	
+	{
+		// ------------- Already sanitized paths ----------------
+		{
+			// directory only - existing
+			const FString ExpectedSanitized = TEXT("./Movies/TmvMediaTestSequence");
+			AddErrorIfFalse(TestSanitizePath(TEXT("./Movies/TmvMediaTestSequence"), ExpectedSanitized, Error), Error);
+			AddErrorIfFalse(TestSanitizePath(TEXT("./Movies/TmvMediaTestSequence/"), ExpectedSanitized, Error), Error);
+			AddErrorIfFalse(TestSanitizePath(TEXT(".\\Movies\\TmvMediaTestSequence"), ExpectedSanitized, Error), Error);
+			AddErrorIfFalse(TestSanitizePath(TEXT(".\\Movies\\TmvMediaTestSequence\\"), ExpectedSanitized, Error), Error);
+			AddErrorIfFalse(TestSanitizePath(TEXT("\".\\Movies\\TmvMediaTestSequence\""), ExpectedSanitized, Error), Error);
+
+			const FString ExpectedFull = FPaths::Combine(FPaths::ConvertRelativePathToFull(FPaths::ProjectDir()), TEXT("Movies/TmvMediaTestSequence"));
+			AddErrorIfFalse(TestConvertPathToFull(ExpectedSanitized, ExpectedFull, Error), Error);
+		}		
+		{
+			// directory only - non-existing
+			const FString ExpectedSanitized = TEXT("./Movies/TmvMediaTestSequenceNotExisting");
+			AddErrorIfFalse(TestSanitizePath(TEXT("./Movies/TmvMediaTestSequenceNotExisting"), ExpectedSanitized, Error), Error);
+			AddErrorIfFalse(TestSanitizePath(TEXT("./Movies/TmvMediaTestSequenceNotExisting/"), ExpectedSanitized, Error), Error);
+			AddErrorIfFalse(TestSanitizePath(TEXT(".\\Movies\\TmvMediaTestSequenceNotExisting"), ExpectedSanitized, Error), Error);
+			AddErrorIfFalse(TestSanitizePath(TEXT(".\\Movies\\TmvMediaTestSequenceNotExisting\\"), ExpectedSanitized, Error), Error);
+			AddErrorIfFalse(TestSanitizePath(TEXT("\".\\Movies\\TmvMediaTestSequenceNotExisting\""), ExpectedSanitized, Error), Error);
+
+			const FString ExpectedFull = FPaths::Combine(FPaths::ConvertRelativePathToFull(FPaths::ProjectDir()), TEXT("Movies/TmvMediaTestSequenceNotExisting"));
+			AddErrorIfFalse(TestConvertPathToFull(ExpectedSanitized, ExpectedFull, Error), Error);
+		}
+		{
+			// Already sanitized - with filename - existing
+			const FString ExpectedSanitized = TEXT("./Movies/TmvMediaTestSequence/Image0000.exr");
+			AddErrorIfFalse(TestSanitizePath(TEXT("./Movies/TmvMediaTestSequence/Image0000.exr"), ExpectedSanitized, Error), Error);
+			AddErrorIfFalse(TestSanitizePath(TEXT(".\\Movies\\TmvMediaTestSequence\\Image0000.exr"), ExpectedSanitized, Error), Error);
+			AddErrorIfFalse(TestSanitizePath(TEXT("\".\\Movies\\TmvMediaTestSequence\\Image0000.exr\""), ExpectedSanitized, Error), Error);
+
+			const FString ExpectedFull = FPaths::Combine(FPaths::ConvertRelativePathToFull(FPaths::ProjectDir()), TEXT("Movies/TmvMediaTestSequence/Image0000.exr"));
+			AddErrorIfFalse(TestConvertPathToFull(ExpectedSanitized, ExpectedFull, Error), Error);
+		}
+		{
+			// Already sanitized - with filename - non-existing
+			const FString ExpectedSanitized = TEXT("./Movies/TmvMediaTestSequenceNotExisting/Image0000.exr");
+			AddErrorIfFalse(TestSanitizePath(TEXT("./Movies/TmvMediaTestSequenceNotExisting/Image0000.exr"), ExpectedSanitized, Error), Error);
+			AddErrorIfFalse(TestSanitizePath(TEXT(".\\Movies\\TmvMediaTestSequenceNotExisting\\Image0000.exr"), ExpectedSanitized, Error), Error);
+			AddErrorIfFalse(TestSanitizePath(TEXT("\".\\Movies\\TmvMediaTestSequenceNotExisting\\Image0000.exr\""), ExpectedSanitized, Error), Error);
+
+			const FString ExpectedFull = FPaths::Combine(FPaths::ConvertRelativePathToFull(FPaths::ProjectDir()), TEXT("Movies/TmvMediaTestSequenceNotExisting/Image0000.exr"));
+			AddErrorIfFalse(TestConvertPathToFull(ExpectedSanitized, ExpectedFull, Error), Error);
+		}
+	}
+	
+	{
+		// ------------- Project Relative ----------------
+		{
+			// directory only - existing
+			const FString ExpectedSanitized = TEXT("./Movies/TmvMediaTestSequence");
+			const FString ProjectRelative = FPaths::Combine(FPaths::ProjectDir(), TEXT("Movies/TmvMediaTestSequence"));
+			AddErrorIfFalse(TestSanitizePath(ProjectRelative, ExpectedSanitized, Error), Error);
+			AddErrorIfFalse(TestSanitizePath(FPaths::ConvertRelativePathToFull(ProjectRelative), ExpectedSanitized, Error), Error);
+
+			const FString ExpectedFull = FPaths::ConvertRelativePathToFull(ProjectRelative);
+			AddErrorIfFalse(TestConvertPathToFull(ExpectedSanitized, ExpectedFull, Error), Error);
+		}
+		{
+			// directory only - non-existing ---- Experimental
+			const FString ExpectedSanitized = TEXT("./Movies/TmvMediaTestSequenceNotExisting");
+			const FString ProjectRelative = FPaths::Combine(FPaths::ProjectDir(), TEXT("Movies/TmvMediaTestSequenceNotExisting"));
+			AddErrorIfFalse(TestSanitizePath(ProjectRelative, ExpectedSanitized, Error), Error);
+			AddErrorIfFalse(TestSanitizePath(FPaths::ConvertRelativePathToFull(ProjectRelative), ExpectedSanitized, Error), Error);
+
+			const FString ExpectedFull = FPaths::ConvertRelativePathToFull(ProjectRelative);
+			AddErrorIfFalse(TestConvertPathToFull(ExpectedSanitized, ExpectedFull, Error), Error);
+		}
+		{
+			// with filename - existing
+			const FString ExpectedSanitized = TEXT("./Movies/TmvMediaTestSequence/Image0000.exr");
+			const FString ProjectRelative = FPaths::Combine(FPaths::ProjectDir(), TEXT("Movies/TmvMediaTestSequence/Image0000.exr"));
+			AddErrorIfFalse(TestSanitizePath(ProjectRelative, ExpectedSanitized, Error), Error);
+			AddErrorIfFalse(TestSanitizePath(FPaths::ConvertRelativePathToFull(ProjectRelative), ExpectedSanitized, Error), Error);
+
+			const FString ExpectedFull = FPaths::ConvertRelativePathToFull(ProjectRelative);
+			AddErrorIfFalse(TestConvertPathToFull(ExpectedSanitized, ExpectedFull, Error), Error);
+		}
+		{
+			// with filename - non-existing ---- Experimental
+			const FString ExpectedSanitized = TEXT("./Movies/TmvMediaTestSequenceNotExisting/Image0000.exr");
+			const FString ProjectRelative = FPaths::Combine(FPaths::ProjectDir(), TEXT("Movies/TmvMediaTestSequenceNotExisting/Image0000.exr"));
+			AddErrorIfFalse(TestSanitizePath(ProjectRelative, ExpectedSanitized, Error), Error);
+			AddErrorIfFalse(TestSanitizePath(FPaths::ConvertRelativePathToFull(ProjectRelative), ExpectedSanitized, Error), Error);
+
+			const FString ExpectedFull = FPaths::ConvertRelativePathToFull(ProjectRelative);
+			AddErrorIfFalse(TestConvertPathToFull(ExpectedSanitized, ExpectedFull, Error), Error);
+		}
+	}
+
+	{
+		// ------------- Content relative -------------
+		// Note: Non-existing content relative can't be done with this syntax. Testing only existing files and directories.
+		{
+			// directory only - existing
+			const FString ExpectedSanitized = TEXT("./Content/Movies/TmvMediaTestSequence");
+			const FString ContentRelative = FPaths::Combine(FPaths::ProjectContentDir(), TEXT("Movies/TmvMediaTestSequence"));
+			AddErrorIfFalse(TestSanitizePath(ContentRelative, ExpectedSanitized, Error), Error);
+			AddErrorIfFalse(TestSanitizePath(FPaths::ConvertRelativePathToFull(ContentRelative), ExpectedSanitized, Error), Error);
+
+			const FString ExpectedFull = FPaths::ConvertRelativePathToFull(ContentRelative);
+			AddErrorIfFalse(TestConvertPathToFull(ExpectedSanitized, ExpectedFull, Error), Error);
+		}
+		{
+			// With filename - existing
+			const FString ExpectedSanitized = TEXT("./Content/Movies/TmvMediaTestSequence/Image0000.exr");
+			const FString ContentRelative = FPaths::Combine(FPaths::ProjectContentDir(), TEXT("Movies/TmvMediaTestSequence/Image0000.exr"));
+			AddErrorIfFalse(TestSanitizePath(ContentRelative, ExpectedSanitized, Error), Error);
+			AddErrorIfFalse(TestSanitizePath(FPaths::ConvertRelativePathToFull(ContentRelative), ExpectedSanitized, Error), Error);
+
+			const FString ExpectedFull = FPaths::ConvertRelativePathToFull(ContentRelative);
+			AddErrorIfFalse(TestConvertPathToFull(ExpectedSanitized, ExpectedFull, Error), Error);
+		}
+	}
+
+	{
+		// ------------- Outside of project (ex: engine) -------------
+		// expect absolute path
+		{
+			// directory only - existing
+			const FString EngineRelative = FPaths::Combine(FPaths::EngineDir(), TEXT("Movies/TmvMediaTestSequence"));
+			const FString EngineAbsolute = FPaths::ConvertRelativePathToFull(EngineRelative);
+			const FString ExpectedSanitized = EngineAbsolute;
+			
+			AddErrorIfFalse(TestSanitizePath(EngineRelative, ExpectedSanitized, Error), Error);
+			AddErrorIfFalse(TestSanitizePath(EngineAbsolute, ExpectedSanitized, Error), Error);
+
+			AddErrorIfFalse(TestConvertPathToFull(ExpectedSanitized, EngineAbsolute, Error), Error);
+		}
+
+		{
+			// directory only - non-existing
+			const FString EngineRelative = FPaths::Combine(FPaths::EngineDir(), TEXT("Movies/TmvMediaTestSequenceNotExisting"));
+			const FString EngineAbsolute = FPaths::ConvertRelativePathToFull(EngineRelative);
+			const FString ExpectedSanitized = EngineAbsolute;
+			
+			AddErrorIfFalse(TestSanitizePath(EngineRelative, ExpectedSanitized, Error), Error);
+			AddErrorIfFalse(TestSanitizePath(EngineAbsolute, ExpectedSanitized, Error), Error);
+
+			AddErrorIfFalse(TestConvertPathToFull(ExpectedSanitized, EngineAbsolute, Error), Error);
+		}
+		
+		{
+			// with filename - existing
+			const FString EngineRelative = FPaths::Combine(FPaths::EngineDir(), TEXT("Movies/TmvMediaTestSequence/Image0000.exr"));
+			const FString EngineAbsolute = FPaths::ConvertRelativePathToFull(EngineRelative);
+			const FString ExpectedSanitized = EngineAbsolute;
+			
+			AddErrorIfFalse(TestSanitizePath(EngineRelative, ExpectedSanitized, Error), Error);
+			AddErrorIfFalse(TestSanitizePath(EngineAbsolute, ExpectedSanitized, Error), Error);
+
+			AddErrorIfFalse(TestConvertPathToFull(ExpectedSanitized, EngineAbsolute, Error), Error);
+		}
+
+		{
+			// with filename - non-existing
+			const FString EngineRelative = FPaths::Combine(FPaths::EngineDir(), TEXT("Movies/TmvMediaTestSequenceNotExisting/Image0000.exr"));
+			const FString EngineAbsolute = FPaths::ConvertRelativePathToFull(EngineRelative);
+			const FString ExpectedSanitized = EngineAbsolute;
+			
+			AddErrorIfFalse(TestSanitizePath(EngineRelative, ExpectedSanitized, Error), Error);
+			AddErrorIfFalse(TestSanitizePath(EngineAbsolute, ExpectedSanitized, Error), Error);
+
+			AddErrorIfFalse(TestConvertPathToFull(ExpectedSanitized, EngineAbsolute, Error), Error);
+		}
+	}
+
+	{
+		// ------------- Project relative with Token -------------
+		// Note: also testing "sanitization" of paths with tokens.
+		{
+			// directory only - existing
+			FString Expected = TEXT("{project_dir}/Movies/TmvMediaTestSequence");
+
+			AddErrorIfFalse(TestSanitizePath(TEXT("{project_dir}/Movies/TmvMediaTestSequence"), Expected, Error), Error);
+			AddErrorIfFalse(TestSanitizePath(TEXT("{project_dir}/Movies/TmvMediaTestSequence/"), Expected, Error), Error);
+			AddErrorIfFalse(TestSanitizePath(TEXT("{project_dir}\\Movies\\TmvMediaTestSequence"), Expected, Error), Error);
+			AddErrorIfFalse(TestSanitizePath(TEXT("{project_dir}\\Movies\\TmvMediaTestSequence\\"), Expected, Error), Error);
+			AddErrorIfFalse(TestSanitizePath(TEXT("\"{project_dir}\\Movies\\TmvMediaTestSequence\""), Expected, Error), Error);
+
+			AddErrorIfFalse(TestSanitizePath(Expected, Expected, Error), Error);
+
+			const FString ExpectedFull = FPaths::Combine(FPaths::ConvertRelativePathToFull(FPaths::ProjectDir()), TEXT("Movies/TmvMediaTestSequence"));
+			AddErrorIfFalse(TestConvertPathToFull(Expected, ExpectedFull, Error), Error);
+		}
+		{
+			// directory only - non-existing
+			FString Expected = TEXT("{project_dir}/Movies/TmvMediaTestSequenceNotExisting");
+
+			AddErrorIfFalse(TestSanitizePath(TEXT("{project_dir}/Movies/TmvMediaTestSequenceNotExisting"), Expected, Error), Error);
+			AddErrorIfFalse(TestSanitizePath(TEXT("{project_dir}/Movies/TmvMediaTestSequenceNotExisting/"), Expected, Error), Error);
+			AddErrorIfFalse(TestSanitizePath(TEXT("{project_dir}\\Movies\\TmvMediaTestSequenceNotExisting"), Expected, Error), Error);
+			AddErrorIfFalse(TestSanitizePath(TEXT("{project_dir}\\Movies\\TmvMediaTestSequenceNotExisting\\"), Expected, Error), Error);
+			AddErrorIfFalse(TestSanitizePath(TEXT("\"{project_dir}\\Movies\\TmvMediaTestSequenceNotExisting\""), Expected, Error), Error);
+
+			AddErrorIfFalse(TestSanitizePath(Expected, Expected, Error), Error);
+
+			const FString ExpectedFull = FPaths::Combine(FPaths::ConvertRelativePathToFull(FPaths::ProjectDir()), TEXT("Movies/TmvMediaTestSequenceNotExisting"));
+			AddErrorIfFalse(TestConvertPathToFull(Expected, ExpectedFull, Error), Error);
+		}
+		{
+			// with filename - existing
+			FString Expected = TEXT("{project_dir}/Movies/TmvMediaTestSequence/Image0000.exr");
+			AddErrorIfFalse(TestSanitizePath(TEXT("{project_dir}/Movies/TmvMediaTestSequence/Image0000.exr"), Expected, Error), Error);
+			AddErrorIfFalse(TestSanitizePath(TEXT("{project_dir}\\Movies\\TmvMediaTestSequence\\Image0000.exr"), Expected, Error), Error);
+			AddErrorIfFalse(TestSanitizePath(TEXT("\"{project_dir}\\Movies\\TmvMediaTestSequence\\Image0000.exr\""), Expected, Error), Error);
+
+			const FString ExpectedFull = FPaths::Combine(FPaths::ConvertRelativePathToFull(FPaths::ProjectDir()), TEXT("Movies/TmvMediaTestSequence/Image0000.exr"));
+			AddErrorIfFalse(TestConvertPathToFull(Expected, ExpectedFull, Error), Error);
+		}
+		{
+			// with filename - non-existing
+			FString Expected = TEXT("{project_dir}/Movies/TmvMediaTestSequenceNotExisting/Image0000.exr");
+			AddErrorIfFalse(TestSanitizePath(TEXT("{project_dir}/Movies/TmvMediaTestSequenceNotExisting/Image0000.exr"), Expected, Error), Error);
+			AddErrorIfFalse(TestSanitizePath(TEXT("{project_dir}\\Movies\\TmvMediaTestSequenceNotExisting\\Image0000.exr"), Expected, Error), Error);
+			AddErrorIfFalse(TestSanitizePath(TEXT("\"{project_dir}\\Movies\\TmvMediaTestSequenceNotExisting\\Image0000.exr\""), Expected, Error), Error);
+
+			const FString ExpectedFull = FPaths::Combine(FPaths::ConvertRelativePathToFull(FPaths::ProjectDir()), TEXT("Movies/TmvMediaTestSequenceNotExisting/Image0000.exr"));
+			AddErrorIfFalse(TestConvertPathToFull(Expected, ExpectedFull, Error), Error);
+		}
+	}
+
+	{
+		// ------------- Engine relative with Token -------------
+		{
+			// directory only - existing
+			FString Expected = TEXT("{engine_dir}/Movies/TmvMediaTestSequence");
+			AddErrorIfFalse(TestSanitizePath(Expected, Expected, Error), Error);
+
+			const FString ExpectedFull = FPaths::Combine(FPaths::ConvertRelativePathToFull(FPaths::EngineDir()), TEXT("Movies/TmvMediaTestSequence"));
+			AddErrorIfFalse(TestConvertPathToFull(Expected, ExpectedFull, Error), Error);
+		}
+		{
+			// directory only - non-existing
+			FString Expected = TEXT("{engine_dir}/Movies/TmvMediaTestSequenceNotExisting");
+			AddErrorIfFalse(TestSanitizePath(Expected, Expected, Error), Error);
+
+			const FString ExpectedFull = FPaths::Combine(FPaths::ConvertRelativePathToFull(FPaths::EngineDir()), TEXT("Movies/TmvMediaTestSequenceNotExisting"));
+			AddErrorIfFalse(TestConvertPathToFull(Expected, ExpectedFull, Error), Error);
+		}
+		{
+			// with filename - existing
+			FString Expected = TEXT("{engine_dir}/Movies/TmvMediaTestSequence/Image0000.exr");
+			AddErrorIfFalse(TestSanitizePath(Expected, Expected, Error), Error);
+
+			const FString ExpectedFull = FPaths::Combine(FPaths::ConvertRelativePathToFull(FPaths::EngineDir()), TEXT("Movies/TmvMediaTestSequence/Image0000.exr"));
+			AddErrorIfFalse(TestConvertPathToFull(Expected, ExpectedFull, Error), Error);
+		}
+		{
+			// with filename - existing
+			FString Expected = TEXT("{engine_dir}/Movies/TmvMediaTestSequenceNotExisting/Image0000.exr");
+			AddErrorIfFalse(TestSanitizePath(Expected, Expected, Error), Error);
+
+			const FString ExpectedFull = FPaths::Combine(FPaths::ConvertRelativePathToFull(FPaths::EngineDir()), TEXT("Movies/TmvMediaTestSequenceNotExisting/Image0000.exr"));
+			AddErrorIfFalse(TestConvertPathToFull(Expected, ExpectedFull, Error), Error);
+		}
+	}
+
+	return !HasAnyErrors();
+}
+
+#endif
diff --git a/Engine/Plugins/Media/TmvMedia/Source/TmvMedia/Private/Transcoder/TmvMediaFileSequenceFrameProducer.cpp b/Engine/Plugins/Media/TmvMedia/Source/TmvMedia/Private/Transcoder/TmvMediaFileSequenceFrameProducer.cpp
--- a/Engine/Plugins/Media/TmvMedia/Source/TmvMedia/Private/Transcoder/TmvMediaFileSequenceFrameProducer.cpp
+++ b/Engine/Plugins/Media/TmvMedia/Source/TmvMedia/Private/Transcoder/TmvMediaFileSequenceFrameProducer.cpp
@@ -122,8 +122,9 @@ void UTmvMediaFileSequenceFrameProducer::ProcessAllImages(UTmvMediaTranscodeJob*
 
 	if (!FPaths::DirectoryExists(SequencePath))
 	{
-		UE_LOG(LogTmvMedia, Error, TEXT("Directory \"%s\" does not exist."), *SequencePath);
-		UTmvMediaTranscodeJob::SafeReportError(InParentJob);
+		const FText Message = FText::Format(LOCTEXT("DirectoryNotFound", "Directory \"{0}\" does not exist"), FText::FromString(SequencePath));
+		UE_LOG(LogTmvMedia, Error, TEXT("FileSequenceFrameProducer: %s."), *Message.ToString());
+		UTmvMediaTranscodeJob::SafeReportError(InParentJob, Message);
 		return;
 	}
 
@@ -135,8 +136,10 @@ void UTmvMediaFileSequenceFrameProducer::ProcessAllImages(UTmvMediaTranscodeJob*
 
 	if (FoundFiles.Num() == 0)
 	{
-		UE_LOG(LogTmvMedia, Error, TEXT("No files to import in directory \"%s\" (Extension: %s)."), *SequencePath, *FileExtension);
-		UTmvMediaTranscodeJob::SafeReportError(InParentJob);
+		const FText Message = FText::Format(LOCTEXT("NoFiles", "No files to import in directory \"{0}\" (Extension: {1})"),
+			FText::FromString(SequencePath), FText::FromString(FileExtension));
+		UE_LOG(LogTmvMedia, Error, TEXT("FileSequenceFrameProducer: %s."), *Message.ToString());
+		UTmvMediaTranscodeJob::SafeReportError(InParentJob, Message);
 		return;
 	}
 
@@ -148,8 +151,9 @@ void UTmvMediaFileSequenceFrameProducer::ProcessAllImages(UTmvMediaTranscodeJob*
 
 	if (ImageFormat == EImageFormat::Invalid)
 	{
-		UE_LOG(LogTmvMedia, Error, TEXT("Invalid file format \"%s\""), *Ext);
-		UTmvMediaTranscodeJob::SafeReportError(InParentJob);
+		const FText Message = FText::Format(LOCTEXT("InvalidFileFormat", "Invalid file format \"{0}\""), FText::FromString(Ext));
+		UE_LOG(LogTmvMedia, Error, TEXT("FileSequenceFrameProducer: %s."), *Message.ToString());
+		UTmvMediaTranscodeJob::SafeReportError(InParentJob, Message);
 		return;
 	}
 
diff --git a/Engine/Plugins/Media/TmvMedia/Source/TmvMedia/Private/Transcoder/TmvMediaFrameConverter.cpp b/Engine/Plugins/Media/TmvMedia/Source/TmvMedia/Private/Transcoder/TmvMediaFrameConverter.cpp
--- a/Engine/Plugins/Media/TmvMedia/Source/TmvMedia/Private/Transcoder/TmvMediaFrameConverter.cpp
+++ b/Engine/Plugins/Media/TmvMedia/Source/TmvMedia/Private/Transcoder/TmvMediaFrameConverter.cpp
@@ -16,6 +16,8 @@
 
 #include UE_INLINE_GENERATED_CPP_BY_NAME(TmvMediaFrameConverter)
 
+#define LOCTEXT_NAMESPACE "TmvMediaFrameConverter"
+
 namespace UE::TmvMedia::FrameConverter
 {
 	/** Convert a row from float16 to float32 quadruplets. */
@@ -619,8 +621,8 @@ namespace UE::TmvMedia::FrameConverter
 
 		for (int32 PlaneIndex = 0; PlaneIndex < InMipInfoSrc.Planes.Num(); ++PlaneIndex)
 		{
-			auto& SrcPlane = InMipInfoSrc.Planes[PlaneIndex];
-			auto& DstPlane = InMipInfoDst.Planes[PlaneIndex];
+			const FTmvMediaFramePlaneInfo& SrcPlane = InMipInfoSrc.Planes[PlaneIndex];
+			const FTmvMediaFramePlaneInfo& DstPlane = InMipInfoDst.Planes[PlaneIndex];
 			if (SrcPlane.NumComponents != DstPlane.NumComponents
 				|| SrcPlane.BitDepth != DstPlane.BitDepth
 				|| SrcPlane.Type != DstPlane.Type
@@ -671,8 +673,9 @@ void UTmvMediaFrameConverter::ReceiveMips(UTmvMediaTranscodeJob* InParentJob, TU
 		}
 		else
 		{
-			UE_LOG(LogTmvMedia, Error, TEXT("Frame Encoder Stage: Invalid Source Mips Parameter: No images or buffers to get info from."));
-			UTmvMediaTranscodeJob::SafeReportError(InParentJob);
+			const FText Message = LOCTEXT("InvalidMipsParam", "Invalid Source Mips Parameter: No images or buffers to get info from");			
+			UE_LOG(LogTmvMedia, Error, TEXT("FrameConverter: %s."), *Message.ToString());
+			UTmvMediaTranscodeJob::SafeReportError(InParentJob, Message);
 			return;
 		}
 
@@ -693,16 +696,18 @@ void UTmvMediaFrameConverter::ReceiveMips(UTmvMediaTranscodeJob* InParentJob, TU
 					TRACE_CPUPROFILER_EVENT_SCOPE(FrameConverter::RequestAllocation);	// (takes 15 ms, todo: pool this)
 					if (!NewMip->RequestAllocation(EncoderFrameMipInfo[MipIndex]))
 					{
-						UE_LOG(LogTmvMedia, Error, TEXT("Frame Converter: Failed to allocate mip %d."), MipIndex);
-						UTmvMediaTranscodeJob::SafeReportError(InParentJob);
+						const FText Message = FText::Format(LOCTEXT("FailedAllocateMip", "Failed to allocate mip {0}"), FText::AsNumber(MipIndex));			
+						UE_LOG(LogTmvMedia, Error, TEXT("FrameConverter: %s."), *Message.ToString());
+						UTmvMediaTranscodeJob::SafeReportError(InParentJob, Message);
 						return;
 					}
 				}
 
 				if (!UE::TmvMedia::FrameConverter::Convert(InMips->MipBuffers[MipIndex].Get(), NewMip.Get()))
 				{
-					UE_LOG(LogTmvMedia, Error, TEXT("Frame Converter: Failed to convert mip %d."), MipIndex);
-					UTmvMediaTranscodeJob::SafeReportError(InParentJob);
+					const FText Message = FText::Format(LOCTEXT("FailedConvertMip", "Failed to convert mip {0}"), FText::AsNumber(MipIndex));			
+					UE_LOG(LogTmvMedia, Error, TEXT("FrameConverter: %s."), *Message.ToString());
+					UTmvMediaTranscodeJob::SafeReportError(InParentJob, Message);
 					return;
 				}
 				
@@ -725,4 +730,6 @@ void UTmvMediaFrameConverter::ReceiveMips(UTmvMediaTranscodeJob* InParentJob, TU
 	{
 		UE_LOG(LogTmvMedia, Warning, TEXT("Transcoding Job: No mips where converted."));
 	}
-}
\ No newline at end of file
+}
+
+#undef LOCTEXT_NAMESPACE
diff --git a/Engine/Plugins/Media/TmvMedia/Source/TmvMedia/Private/Transcoder/TmvMediaPlayerFrameProducer.cpp b/Engine/Plugins/Media/TmvMedia/Source/TmvMedia/Private/Transcoder/TmvMediaPlayerFrameProducer.cpp
--- a/Engine/Plugins/Media/TmvMedia/Source/TmvMedia/Private/Transcoder/TmvMediaPlayerFrameProducer.cpp
+++ b/Engine/Plugins/Media/TmvMedia/Source/TmvMedia/Private/Transcoder/TmvMediaPlayerFrameProducer.cpp
@@ -22,6 +22,32 @@
 
 #define LOCTEXT_NAMESPACE "TmvMediaPlayerFrameProducer"
 
+/** Timespan logging utility functions. */
+namespace UE::TmvMedia
+{
+	/**
+	 * Utility function to convert a timespan to pretty formatted string for logging. 
+	 */
+	FString ToString(const FTimespan& InTimespan)
+	{
+		// The selected format displays the fraction as ticks "t" to be consistent with
+		// LogMediaAssets (MediaPlayer), LogMediaUtils (PlayerFacade) and sequence media track.
+		// LogElectraPlayer is a mix of seconds (playback controls) or ticks (Presentation time stamps).
+		return InTimespan.ToString(TEXT("%h:%m:%s.%t"));
+	}
+
+	/** Utility function to convert a timespan bound to formatted string for logging. */
+	FString ToString(const TRange<FTimespan>::BoundsType& InBounds)
+	{
+		// Only call GetValue() if the bound is closed.
+		if (InBounds.IsClosed())
+		{
+			return ToString(InBounds.GetValue());
+		}
+		return TEXT("<open>");
+	}
+}
+
 UTmvMediaPlayerFrameProducer::UTmvMediaPlayerFrameProducer()
 {
 	bHasAsyncQueue = true;
@@ -50,7 +76,11 @@ bool UTmvMediaPlayerFrameProducer::Start(UTmvMediaTranscodeJob* InParentJob)
 		MediaSource = UMediaSource::SpawnMediaSourceForString(InParentJob->Settings.GetAbsoluteInputPath(), this);
 		if (MediaSource == nullptr)
 		{
-			UE_LOG(LogTmvMedia, Error, TEXT("Failed to create a media source from path \"%s\""), *InParentJob->Settings.GetInputPath());
+			const FText Message = FText::Format(LOCTEXT("FailedToCreateMediaSource", 
+				"Failed to create a media source from path \"{0}\""),
+				FText::FromString(InParentJob->Settings.GetInputPath()));
+			UTmvMediaTranscodeJob::SafeReportError(InParentJob, Message);
+			UE_LOG(LogTmvMedia, Error, TEXT("MediaPlayerFrameProducer::Start: %s."), *Message.ToString());
 			return false;
 		}
 	}
@@ -59,7 +89,11 @@ bool UTmvMediaPlayerFrameProducer::Start(UTmvMediaTranscodeJob* InParentJob)
 		MediaSource = InParentJob->Settings.InputMediaSource.LoadSynchronous();
 		if (MediaSource == nullptr)
 		{
-			UE_LOG(LogTmvMedia, Error, TEXT("Invalid Media Source Asset: \"%s\""), *InParentJob->Settings.InputMediaSource.ToString());
+			const FText Message = FText::Format(LOCTEXT("FailedToCreateMediaSource", 
+				"Invalid Media Source Asset: \"{0}\""),
+				FText::FromString(InParentJob->Settings.InputMediaSource.ToString()));
+			UTmvMediaTranscodeJob::SafeReportError(InParentJob, Message);
+			UE_LOG(LogTmvMedia, Error, TEXT("MediaPlayerFrameProducer::Start: %s."), *Message.ToString());
 			return false;
 		}
 	}
@@ -139,6 +173,8 @@ void UTmvMediaPlayerFrameProducer::RequestStop(UTmvMediaTranscodeJob* InParentJo
 
 void UTmvMediaPlayerFrameProducer::Tick(UTmvMediaTranscodeJob* InParentJob, const FTmvMediaTranscodeJobTime& InTime)
 {
+	using namespace UE::TmvMedia;
+
 	if (!InParentJob)
 	{
 		UE_LOG(LogTmvMedia, Error, TEXT("UTmvMediaPlayerFrameProducer::Tick: InParentJob is null. Stopping Stage."));
@@ -218,7 +254,7 @@ void UTmvMediaPlayerFrameProducer::Tick(UTmvMediaTranscodeJob* InParentJob, cons
 					MediaStartTime = CurrentTime = RenderedFrameTime;
 					UE_LOG(LogTmvMedia, Verbose,
 						TEXT("Transcode Frame Producer: Received First Frame %s"),
-						*MediaStartTime.ToString());
+						*ToString(MediaStartTime));
 
 					const FTimespan RenderedFrameDuration = ProducerResource->GetLastRenderedSampleDuration();
 
@@ -227,7 +263,7 @@ void UTmvMediaPlayerFrameProducer::Tick(UTmvMediaTranscodeJob* InParentJob, cons
 					{
 						UE_LOG(LogTmvMedia, Warning,
 							TEXT("Transcode Frame Producer: Frame duration mismatch: %s (player) vs %s (sample)."),
-							*MediaFrameDuration.ToString(), *RenderedFrameDuration.ToString());
+							*ToString(MediaFrameDuration), *ToString(RenderedFrameDuration));
 					}
 					
 					// We need to guard in case the provided frame duration is zero, getting the job stuck.
@@ -237,7 +273,7 @@ void UTmvMediaPlayerFrameProducer::Tick(UTmvMediaTranscodeJob* InParentJob, cons
 					{
 						UE_LOG(LogTmvMedia, Error,
 							TEXT("Transcode Frame Producer: Can't safely block/advance player without a frame duration: %s (player) vs %s (sample)."),
-							*MediaFrameDuration.ToString(), *RenderedFrameDuration.ToString());
+							*ToString(MediaFrameDuration), *ToString(RenderedFrameDuration));
 						CurrentState = EState::Failed;
 						return;
 					}
@@ -286,6 +322,22 @@ void UTmvMediaPlayerFrameProducer::Tick(UTmvMediaTranscodeJob* InParentJob, cons
 		
 		if (ProducerResource)
 		{
+			// Keep updating the stats
+			const FTimespan RenderedSampleTime = ProducerResource->GetLastRenderedSampleTime();
+			if (RenderedSampleTime != FTimespan::MinValue())
+			{
+				const FTimespan RenderedSampleDuration = ProducerResource->GetLastRenderedSampleDuration();
+				const FTimespan FrameDuration = (RenderedSampleDuration > FTimespan::Zero()) ? RenderedSampleDuration : MediaFrameDuration;
+				const FTimespan RelativeTime = RenderedSampleTime - MediaStartTime;
+				if (FrameDuration != FTimespan::Zero() && !MediaPlaybackRange.IsEmpty())
+				{
+					const int32 RenderedFrameIndex = static_cast<int32>(FMath::RoundHalfToEven(RelativeTime.GetTotalSeconds() / FrameDuration.GetTotalSeconds()));
+					const FTimespan SequenceDuration = MediaPlaybackRange.Size<FTimespan>();
+					const int32 TotalFrames = static_cast<int32>(FMath::RoundHalfToEven(SequenceDuration.GetTotalSeconds() / FrameDuration.GetTotalSeconds()));
+					UTmvMediaTranscodeJob::SafeUpdateProgress(InParentJob, RenderedFrameIndex, TotalFrames);
+				}
+			}
+			
 			// Wait for final captures to process.
 			ensure(ProducerResource->GetSubmittedCount() >= 0);
 			if (ProducerResource->GetSubmittedCount() == 0)
@@ -313,6 +365,8 @@ void UTmvMediaPlayerFrameProducer::Tick(UTmvMediaTranscodeJob* InParentJob, cons
 
 void UTmvMediaPlayerFrameProducer::HandlePlaying(UTmvMediaTranscodeJob* InParentJob)
 {
+	using namespace UE::TmvMedia;
+	
 	if (!ProducerResource)
 	{
 		UE_LOG(LogTmvMedia, Error, TEXT("Transcode Frame Producer: Render resource not available."));
@@ -351,7 +405,7 @@ void UTmvMediaPlayerFrameProducer::HandlePlaying(UTmvMediaTranscodeJob* InParent
 	{
 		UE_LOG(LogTmvMedia, Error,
 			TEXT("Transcode Frame Producer: Can't safely block/advance player without a frame duration: %s (player) vs %s (sample)."),
-			*MediaFrameDuration.ToString(), *RenderedSampleDuration.ToString());
+			*ToString(MediaFrameDuration), *ToString(RenderedSampleDuration));
 		CurrentState = EState::Failed;
 		return;
 	}
@@ -385,8 +439,8 @@ void UTmvMediaPlayerFrameProducer::HandlePlaying(UTmvMediaTranscodeJob* InParent
 			// Move to the next frame in the sequence.
 			CurrentTime = RenderedSampleTime + FrameDuration;
 
-			UE_LOG(LogTmvMedia, Verbose, TEXT("Transcode Frame Producer: Moving to next frame: Rendered Time:%s Rendered Duration:%s"),
-				*RenderedSampleTime.ToString(TEXT("%h:%m:%s.%t")), *FrameDuration.ToString(TEXT("%h:%m:%s.%t")));
+			UE_LOG(LogTmvMedia, Verbose, TEXT("Transcode Frame Producer: Moving to next frame: Rendered Time:%s (Index: %d) Rendered Duration:%s"),
+				*ToString(RenderedSampleTime), RenderedFrameIndex, *ToString(FrameDuration));
 
 			if (!MediaPlaybackRange.IsEmpty())
 			{
@@ -397,6 +451,9 @@ void UTmvMediaPlayerFrameProducer::HandlePlaying(UTmvMediaTranscodeJob* InParent
 				}
 				else
 				{
+					UE_LOG(LogTmvMedia, Verbose, TEXT("Transcode Frame Producer: Media completed: Time: %s Range: [%s, %s]"),
+						*ToString(CurrentTime), *ToString(MediaPlaybackRange.GetLowerBound()), *ToString(MediaPlaybackRange.GetUpperBound()));
+
 					// Frame production is completed on the player side, need to wait for the rest of the stages to finish.
 					CurrentState = EState::MediaCompleted;
 					WatchdogTimeRemaining = 10.0f;
@@ -407,13 +464,12 @@ void UTmvMediaPlayerFrameProducer::HandlePlaying(UTmvMediaTranscodeJob* InParent
 		}
 		else
 		{
-			UE_LOG(LogTmvMedia, Verbose, TEXT("Transcode Frame Producer: WAITING  Player Time:%s Rendered Time:%s Playback Range:[%s, %s] Frame:%d"),
-				*CurrentTime.ToString(TEXT("%h:%m:%s.%t")),
-				*RenderedSampleTime.ToString(TEXT("%h:%m:%s.%t")),
-				!MediaPlaybackRange.IsEmpty() ? *MediaPlaybackRange.GetLowerBoundValue().ToString(TEXT("%h:%m:%s.%t")) : TEXT("<unknown>"),
-				!MediaPlaybackRange.IsEmpty() ? *MediaPlaybackRange.GetUpperBoundValue().ToString(TEXT("%h:%m:%s.%t")) : TEXT("<unknown>"),
+			UE_LOG(LogTmvMedia, Verbose, TEXT("Transcode Frame Producer: WAITING  Player Time: %s Rendered Time: %s Playback Range: [%s, %s] Frame: %d"),
+				*ToString(CurrentTime),
+				*ToString(RenderedSampleTime),
+				*ToString(MediaPlaybackRange.GetLowerBound()),
+				*ToString(MediaPlaybackRange.GetUpperBound()),
 				RenderedFrameIndex);
-
 		}
 	}
 }
@@ -515,11 +571,16 @@ void UTmvMediaPlayerFrameProducer::OnMediaOpenFailed(FString InUrl)
 void UTmvMediaPlayerFrameProducer::OnMediaEndReached()
 {
 	UE_LOG(LogTmvMedia, Verbose, TEXT("MediaPlayerFrameProducer Media End Reached"));
-	CurrentState = EState::MediaCompleted;
-	WatchdogTimeRemaining = 10.0f;
+
+	// This event is not precise enough, we will use it only if we don't have a playback range.
+	if (MediaPlaybackRange.IsEmpty())
+	{
+		CurrentState = EState::MediaCompleted;
+		WatchdogTimeRemaining = 10.0f;
 	
-	// Unblock the player.
-	MediaPlayer->SetBlockOnTimeRange(TRange<FTimespan>::Empty());
+		// Unblock the player.
+		MediaPlayer->SetBlockOnTimeRange(TRange<FTimespan>::Empty());
+	}
 }
 
 void UTmvMediaPlayerFrameProducer::OnMediaResumed()
diff --git a/Engine/Plugins/Media/TmvMedia/Source/TmvMedia/Private/Transcoder/TmvMediaPlayerFrameProducerResource.cpp b/Engine/Plugins/Media/TmvMedia/Source/TmvMedia/Private/Transcoder/TmvMediaPlayerFrameProducerResource.cpp
--- a/Engine/Plugins/Media/TmvMedia/Source/TmvMedia/Private/Transcoder/TmvMediaPlayerFrameProducerResource.cpp
+++ b/Engine/Plugins/Media/TmvMedia/Source/TmvMedia/Private/Transcoder/TmvMediaPlayerFrameProducerResource.cpp
@@ -391,7 +391,7 @@ void FTmvMediaPlayerFrameProducerResource::OnMediaTextureRendered_RenderThread(c
 		
 		if (TargetMipInfo.Num())
 		{
-			RenderTargetDesc.NumMips = TargetMipInfo.Num();
+			// Note: The destination render target number of mips matches the media texture, not TargetMipInfo.
 
 			// Make it packed, all components in the same plane, and align with supported formats.
 			TargetMipInfo[0].Planes[0].NumComponents = (TargetMipInfo[0].NumComponents == 1) ? 1 : 4;	// simplify to only support 1 or 4 components.
diff --git a/Engine/Plugins/Media/TmvMedia/Source/TmvMedia/Private/Transcoder/TmvMediaTestFrameEncoder.cpp b/Engine/Plugins/Media/TmvMedia/Source/TmvMedia/Private/Transcoder/TmvMediaTestFrameEncoder.cpp
--- a/Engine/Plugins/Media/TmvMedia/Source/TmvMedia/Private/Transcoder/TmvMediaTestFrameEncoder.cpp
+++ b/Engine/Plugins/Media/TmvMedia/Source/TmvMedia/Private/Transcoder/TmvMediaTestFrameEncoder.cpp
@@ -15,7 +15,7 @@
 
 bool UTmvMediaTestFrameEncoder::Start(UTmvMediaTranscodeJob* InParentJob)
 {
-	const FString OutPath = InParentJob->Settings.OutputPath.Path;
+	const FString OutPath = InParentJob->Settings.GetAbsoluteOutputPath();
 	if (OutPath.IsEmpty())
 	{
 		UE_LOG(LogTmvMedia, Error, TEXT("TestFrameEncoder: OutputPath is empty; cannot write test frames."));
@@ -34,7 +34,7 @@ bool UTmvMediaTestFrameEncoder::Start(UTmvMediaTranscodeJob* InParentJob)
 
 void UTmvMediaTestFrameEncoder::ReceiveMips(UTmvMediaTranscodeJob* InParentJob, TUniquePtr<FTmvMediaFrameMips>&& InMips)
 {
-	FString OutPath = InParentJob->Settings.OutputPath.Path;
+	FString OutPath = InParentJob->Settings.GetAbsoluteOutputPath();
 	int32 MipIndex = 0;
 	for (const FTmvMediaFrameMipBufferHandle& MipBuffer : InMips->MipBuffers)
 	{
diff --git a/Engine/Plugins/Media/TmvMedia/Source/TmvMedia/Private/Transcoder/TmvMediaTmvFrameEncoder.cpp b/Engine/Plugins/Media/TmvMedia/Source/TmvMedia/Private/Transcoder/TmvMediaTmvFrameEncoder.cpp
--- a/Engine/Plugins/Media/TmvMedia/Source/TmvMedia/Private/Transcoder/TmvMediaTmvFrameEncoder.cpp
+++ b/Engine/Plugins/Media/TmvMedia/Source/TmvMedia/Private/Transcoder/TmvMediaTmvFrameEncoder.cpp
@@ -21,6 +21,8 @@
 
 #include UE_INLINE_GENERATED_CPP_BY_NAME(TmvMediaTmvFrameEncoder)
 
+#define LOCTEXT_NAMESPACE "TmvMediaTmvFrameEncoder"
+
 /** Helper to trouble shoot frame conversion. */
 static TAutoConsoleVariable<bool> CVarTmvEncoderEnableDumpRawFrame(
 	TEXT("TmvEncoder.EnableDumpRawFrame"),
@@ -41,7 +43,12 @@ namespace UE::TmvMedia::FrameUtils
 
 		// save only mip 0 for now.
 		constexpr int32 MipIndex = 0;
-		const FString OutPath = InParentJob->Settings.OutputPath.Path;
+		const FString OutPath = InParentJob->Settings.GetAbsoluteOutputPath();
+		if (OutPath.IsEmpty())
+		{
+			UE_LOG(LogTmvMedia, Error, TEXT("WriteFrameToFile Failed: Empty output path."));
+			return false;
+		}
 		FString TmpFileName = FString::Printf(TEXT("%s%05d_mip%02d_tmp"), *InStreamName, InMips.TimeInfo.FrameIndex, MipIndex);
 		return WriteMipBufferToFile(FPaths::Combine(OutPath, FPaths::MakeValidFileName(TmpFileName)), InMips.MipBuffers[MipIndex]);
 	}
@@ -174,7 +181,7 @@ bool UTmvMediaTmvFrameEncoder::Start(UTmvMediaTranscodeJob* InParentJob)
 	}
 	else
 	{
-		const FString OutPath = InParentJob->Settings.OutputPath.Path;
+		const FString OutPath = InParentJob->Settings.GetAbsoluteOutputPath();
 		if (!FPaths::DirectoryExists(OutPath))
 		{
 			IPlatformFile& PlatformFile = FPlatformFileManager::Get().GetPlatformFile();
@@ -289,8 +296,9 @@ void UTmvMediaTmvFrameEncoder::ReceiveMips(UTmvMediaTranscodeJob* InParentJob, T
 	const FTmvMediaTmvEncoderPool::FHandle EncoderHandle = EncoderPoolRef->AcquireEncoder(TEXT(""), EncoderOptions);
 	if (!EncoderHandle.IsValid())
 	{
-		UE_LOG(LogTmvMedia, Error, TEXT("Frame Encoder Stage: Failed to acquire encoder \"%s\"."), *EncoderOptions.Get().GetEncoderName().ToString());
-		UTmvMediaTranscodeJob::SafeReportError(InParentJob);
+		const FText Message = FText::Format(LOCTEXT("FailedAcquireEncoder", "Failed to acquire encoder \"{0}\""), FText::FromName(EncoderOptions.Get().GetEncoderName()));
+		UE_LOG(LogTmvMedia, Error, TEXT("Frame Encoder Stage: %s."), *Message.ToString());
+		UTmvMediaTranscodeJob::SafeReportError(InParentJob, Message);
 		return;
 	}
 
@@ -298,7 +306,16 @@ void UTmvMediaTmvFrameEncoder::ReceiveMips(UTmvMediaTranscodeJob* InParentJob, T
 	TSharedPtr<TArray64<uint8>> AccessUnitBuffer = MakeShared<TArray64<uint8>>();
 
 	{
-		FString OutputPath = InParentJob->Settings.OutputPath.Path;
+		FString OutputPath = InParentJob->Settings.GetAbsoluteOutputPath();
+
+		if (OutputPath.IsEmpty())
+		{
+			const FText Message = LOCTEXT("EmptyOutputPath", "Empty Output Path");
+			UE_LOG(LogTmvMedia, Error, TEXT("Frame Encoder Stage: %s."), *Message.ToString());
+			UTmvMediaTranscodeJob::SafeReportError(InParentJob, Message);
+			return;
+		}
+
 		FString FileName = FString::Printf(TEXT("%s%05d.%s"), *StreamName, InMips->TimeInfo.FrameIndex, *EncoderOptions.Get().GetFileSequenceExtension());
 		FString FullFilename = FPaths::Combine(OutputPath, FPaths::MakeValidFileName(FileName));
 
@@ -322,8 +339,9 @@ void UTmvMediaTmvFrameEncoder::ReceiveMips(UTmvMediaTranscodeJob* InParentJob, T
 		ETmvMediaEncoderResult Result = EncoderHandle->Encode(InMips->TimeInfo, AccessUnit, MipRequests);
 		if (Result != ETmvMediaEncoderResult::Success)
 		{
-			UE_LOG(LogTmvMedia, Error, TEXT("Frame Encoder Stage: Failed to encode frame %d."), InMips->TimeInfo.FrameIndex);
-			UTmvMediaTranscodeJob::SafeReportError(InParentJob);
+			const FText Message = FText::Format(LOCTEXT("FailedEncodeFrame", "Failed to encode frame {0}"), FText::AsNumber(InMips->TimeInfo.FrameIndex));
+			UE_LOG(LogTmvMedia, Error, TEXT("Frame Encoder Stage: %s."), *Message.ToString());
+			UTmvMediaTranscodeJob::SafeReportError(InParentJob, Message);
 			return;
 		}
 	}
@@ -335,4 +353,6 @@ void UTmvMediaTmvFrameEncoder::ReceiveMips(UTmvMediaTranscodeJob* InParentJob, T
 			Muxer->ReceiveAccessUnit(InParentJob, StreamId, InMips->TimeInfo, MoveTemp(AccessUnitBuffer));
 		}
 	}
-}
\ No newline at end of file
+}
+
+#undef LOCTEXT_NAMESPACE
\ No newline at end of file
diff --git a/Engine/Plugins/Media/TmvMedia/Source/TmvMedia/Private/Transcoder/TmvMediaTranscodeJob.cpp b/Engine/Plugins/Media/TmvMedia/Source/TmvMedia/Private/Transcoder/TmvMediaTranscodeJob.cpp
--- a/Engine/Plugins/Media/TmvMedia/Source/TmvMedia/Private/Transcoder/TmvMediaTranscodeJob.cpp
+++ b/Engine/Plugins/Media/TmvMedia/Source/TmvMedia/Private/Transcoder/TmvMediaTranscodeJob.cpp
@@ -7,55 +7,13 @@
 #include "TmvMediaLog.h"
 #include "Transcoder/ITmvMediaTranscodeJobManager.h"
 #include "Transcoder/TmvMediaFrameProducer.h"
+#include "Utils/TmvMediaPathUtils.h"
 #include "Utils/TmvMediaUtils.h"
 
 #include UE_INLINE_GENERATED_CPP_BY_NAME(TmvMediaTranscodeJob)
 
 #define LOCTEXT_NAMESPACE "TmvMediaTranscodeJob"
 
-namespace UE::TmvMedia::Transcode
-{
-	// Duplicated from: UImgMediaSource::ExpandSequencePathTokens.
-	// todo: Have some common "file" media source functions to centralize path handling.
-	FString ExpandSequencePathTokens(const FString& InPath)
-	{
-		return InPath
-			.Replace(TEXT("{engine_dir}"), *FPaths::ConvertRelativePathToFull(FPaths::EngineDir()))
-			.Replace(TEXT("{project_dir}"), *FPaths::ConvertRelativePathToFull(FPaths::ProjectDir()))
-			;
-	}
-	
-	FString ConvertPathToFull(const FString& InRelativeOrAbsolutePath)
-	{
-		// Expand the common known path tokens (potentially used in ImgMediaSource).
-		FString RelativeOrAbsolutePathExpanded = ExpandSequencePathTokens(InRelativeOrAbsolutePath);
-		
-		// If the path is already relative.
-		if (FPaths::IsRelative(RelativeOrAbsolutePathExpanded) || RelativeOrAbsolutePathExpanded.StartsWith("./"))
-		{
-			const TArray<FString> PossibleBasePaths =
-			{
-				FPaths::ProjectContentDir(),
-				FPaths::ProjectDir()
-			};
-
-			for (const FString& BasePath : PossibleBasePaths)
-			{
-				const FString BasePathFull = FPaths::ConvertRelativePathToFull(BasePath);
-				const FString PossiblePath = FPaths::ConvertRelativePathToFull(FPaths::Combine(BasePathFull, RelativeOrAbsolutePathExpanded));
-				if (FPaths::FileExists(PossiblePath))
-				{
-					return PossiblePath;
-				}
-			}
-		}
-
-		// Just try to make it absolute (in case it was still relative to binary).
-		return FPaths::ConvertRelativePathToFull(RelativeOrAbsolutePathExpanded);
-	}
-}
-
-
 FString FTmvMediaTranscodeJobSettings::GetInputPath() const
 {
 	if (InputSource == ETmvMediaTranscodeInputSource::File)
@@ -71,7 +29,7 @@ FString FTmvMediaTranscodeJobSettings::GetAbsoluteInputPath() const
 {
 	if (InputSource == ETmvMediaTranscodeInputSource::File)
 	{
-		return UE::TmvMedia::Transcode::ConvertPathToFull(InputPath.FilePath);
+		return UE::TmvMedia::PathUtils::ConvertSanitizedPathToFull(InputPath.FilePath);
 	}
 
 	// Get the media file from the media source.
@@ -85,13 +43,13 @@ FString FTmvMediaTranscodeJobSettings::GetAbsoluteInputPath() const
 
 		if (Url.StartsWith(FilePrefix))
 		{
-			return UE::TmvMedia::Transcode::ConvertPathToFull(Url.RightChop(FilePrefix.Len()));
+			return UE::TmvMedia::PathUtils::ConvertSanitizedPathToFull(Url.RightChop(FilePrefix.Len()));
 		}
 		
 		if (Url.StartsWith(ImgPrefix))
 		{
 			// Img Media only keeps the directory. For image sequence producer, it will then search for files in that directory.
-			return UE::TmvMedia::Transcode::ConvertPathToFull(Url.RightChop(ImgPrefix.Len()));
+			return UE::TmvMedia::PathUtils::ConvertSanitizedPathToFull(Url.RightChop(ImgPrefix.Len()));
 		}
 
 		UE_LOG(LogTmvMedia, Error, TEXT("Media Source \"%s\" is not a supported codec for transcoding."), *Url);
@@ -120,6 +78,11 @@ bool FTmvMediaTranscodeJobSettings::IsOutputPathSet() const
 	return OutputPath.Path.IsEmpty() == false;
 }
 
+FString FTmvMediaTranscodeJobSettings::GetAbsoluteOutputPath() const
+{
+	return UE::TmvMedia::PathUtils::ConvertSanitizedPathToFull(OutputPath.Path);
+}
+
 UTmvMediaTranscodeJob::UTmvMediaTranscodeJob()
 {
 	if(!HasAnyFlags(RF_ClassDefaultObject))
@@ -148,6 +111,9 @@ bool UTmvMediaTranscodeJob::Start(double InCurrentTime)
 	if (!FrameProducer)
 	{
 		UE_LOG(LogTmvMedia, Error, TEXT("Transcode Job for \"%s\" failed to start: missing the Frame Producer stage."), *Settings.GetInputPath());
+		SafeReportError(this, LOCTEXT("ErrorMissingFrameProducer", "Missing Frame Producer Stage"));
+		JobStatus = ETmvMediaTranscodeJobStatus::Stopped;
+		OnStopped(InCurrentTime);
 		return false;
 	}
 
@@ -158,9 +124,9 @@ bool UTmvMediaTranscodeJob::Start(double InCurrentTime)
 		{
 			if (!Stage->Start(this))
 			{
+				bIsError = true;	// Must be set before calling OnStopped.
 				JobStatus = ETmvMediaTranscodeJobStatus::Stopped;
 				OnStopped(InCurrentTime);
-				bIsError = true;
 				return false;
 			}
 		}
@@ -169,9 +135,9 @@ bool UTmvMediaTranscodeJob::Start(double InCurrentTime)
 	// Start the producer last.
 	if (!FrameProducer->Start(this))
 	{
+		bIsError = true;	// Must be set before calling OnStopped.
 		JobStatus = ETmvMediaTranscodeJobStatus::Stopped;
 		OnStopped(InCurrentTime);
-		bIsError = true;
 	}
 
 	return IsRunning();
@@ -190,7 +156,7 @@ void UTmvMediaTranscodeJob::OnStopped(double InCurrentTime)
 			FScopeLock lock(&JobStatsLock);
 			bIsCompleted = JobStats.ProcessedFrame == JobStats.TotalFramesToProcess;
 		}
-		Notification->Close(bIsCompleted);
+		Notification->Close(bIsCompleted && !bIsError);
 	}
 }
 
@@ -241,6 +207,13 @@ void UTmvMediaTranscodeJob::Tick(const FTmvMediaTranscodeJobTime& InTime)
 {
 	if (JobStatus != ETmvMediaTranscodeJobStatus::Stopped)
 	{
+		// Cancel the job in case of error.
+		if (bIsError)
+		{
+			RequestStop(InTime.CurrentTime, ETmvMediaTranscodeJobStopReason::Cancelled);
+			return;
+		}
+		
 		for (UTmvMediaTranscodeStage* Stage : Stages)
 		{
 			Stage->Tick(this, InTime);
diff --git a/Engine/Plugins/Media/TmvMedia/Source/TmvMedia/Private/Transcoder/TmvMediaTranscodeList.cpp b/Engine/Plugins/Media/TmvMedia/Source/TmvMedia/Private/Transcoder/TmvMediaTranscodeList.cpp
--- a/Engine/Plugins/Media/TmvMedia/Source/TmvMedia/Private/Transcoder/TmvMediaTranscodeList.cpp
+++ b/Engine/Plugins/Media/TmvMedia/Source/TmvMedia/Private/Transcoder/TmvMediaTranscodeList.cpp
@@ -41,7 +41,7 @@ void UTmvMediaTranscodeList::InsertItemAt(int32 InItemIndex, const FTmvMediaTran
 	{
 		NewItem.Name = InItem->Name + TEXT("_Dup");
 		NewItem.Settings = InItem->Settings;
-		NewItem.EncoderOptions = InItem->EncoderOptions; // is this duplicating?
+		NewItem.EncoderOptions = InItem->EncoderOptions;
 	}
 	else
 	{
@@ -110,6 +110,45 @@ void UTmvMediaTranscodeList::RemoveItems(TConstArrayView<int32> InItemIndices)
 	}
 }
 
+bool UTmvMediaTranscodeList::ReorderItems(int32 InFromIndex, int32 InToIndex)
+{
+	if (!Items.IsValidIndex(InFromIndex) || InToIndex < 0)
+	{
+		return false;
+	}
+
+	FTmvMediaTranscodeListItem ItemToMove = Items[InFromIndex];
+
+	int32 MovedIndex;
+	
+	if (InToIndex >= Items.Num())
+	{
+		Items.Add(MoveTemp(ItemToMove));
+		MovedIndex = Items.Num() - 1;
+	}
+	else
+	{
+		Items.Insert(MoveTemp(ItemToMove), InToIndex);
+		MovedIndex = InToIndex;
+	}
+	
+	if (MovedIndex > InFromIndex)
+	{
+		--MovedIndex;
+		Items.RemoveAt(InFromIndex);
+	}
+	else
+	{
+		Items.RemoveAt(InFromIndex+1);
+	}
+
+	FTmvMediaTranscodeListItemEventArgs ItemEventArgs;
+	ItemEventArgs.Type = ETmvMediaTranscodeListItemEventType::ItemsReordered;
+	ItemEventArgs.ItemIndices.Add(MovedIndex);
+	OnItemEvent.Broadcast(this, ItemEventArgs);
+	return true;
+}
+
 void UTmvMediaTranscodeList::PostLoad()
 {
 	Super::PostLoad();
diff --git a/Engine/Plugins/Media/TmvMedia/Source/TmvMedia/Private/Transcoder/TmvMediaTranscodeMuxer.cpp b/Engine/Plugins/Media/TmvMedia/Source/TmvMedia/Private/Transcoder/TmvMediaTranscodeMuxer.cpp
--- a/Engine/Plugins/Media/TmvMedia/Source/TmvMedia/Private/Transcoder/TmvMediaTranscodeMuxer.cpp
+++ b/Engine/Plugins/Media/TmvMedia/Source/TmvMedia/Private/Transcoder/TmvMediaTranscodeMuxer.cpp
@@ -12,6 +12,8 @@
 
 #include UE_INLINE_GENERATED_CPP_BY_NAME(TmvMediaTranscodeMuxer)
 
+#define LOCTEXT_NAMESPACE "TmvMediaTranscodeMuxer"
+
 int32 UTmvMediaTranscodeMuxer::OpenStream(UTmvMediaTranscodeJob* InParentJob, const FString& InStreamName, const FString& InExtension)
 {
 	if (!InParentJob)
@@ -21,17 +23,20 @@ int32 UTmvMediaTranscodeMuxer::OpenStream(UTmvMediaTranscodeJob* InParentJob, co
 
 	if (InStreamName.IsEmpty() || InExtension.IsEmpty())
 	{
-		UE_LOG(LogTmvMedia, Error, TEXT("OpenStream: Invalid stream name or extension (Name=\"%s\", Ext=\"%s\")"), *InStreamName, *InExtension);
-		UTmvMediaTranscodeJob::SafeReportError(InParentJob);
+		const FText Message = FText::Format(LOCTEXT("InvalidStream", "Invalid stream name or extension (Name=\"{0}\", Ext=\"{1}\""),
+			FText::FromString(InStreamName), FText::FromString(InExtension));
+		UE_LOG(LogTmvMedia, Error, TEXT("TranscodeMuxer::OpenStream: %s."), *Message.ToString());
+		UTmvMediaTranscodeJob::SafeReportError(InParentJob, Message);
 		return INDEX_NONE;
 	}
 
-	const FString OutPath = InParentJob->Settings.OutputPath.Path;
+	const FString OutPath = InParentJob->Settings.GetAbsoluteOutputPath();
 
 	if (OutPath.IsEmpty())
 	{
-		UE_LOG(LogTmvMedia, Error, TEXT("OpenStream: OutputPath is empty."));
-		UTmvMediaTranscodeJob::SafeReportError(InParentJob);
+		const FText Message = LOCTEXT("EmptyOutputPath", "OutputPath is empty");
+		UE_LOG(LogTmvMedia, Error, TEXT("TranscodeMuxer::OpenStream: %s."), *Message.ToString());
+		UTmvMediaTranscodeJob::SafeReportError(InParentJob, Message);
 		return INDEX_NONE;
 	}
 
@@ -39,8 +44,9 @@ int32 UTmvMediaTranscodeMuxer::OpenStream(UTmvMediaTranscodeJob* InParentJob, co
 	{
 		if (!FPlatformFileManager::Get().GetPlatformFile().CreateDirectoryTree(*OutPath))
 		{
-			UE_LOG(LogTmvMedia, Error, TEXT("OpenStream: Failed to create output directory \"%s\""), *OutPath);
-			UTmvMediaTranscodeJob::SafeReportError(InParentJob);
+			const FText Message = FText::Format(LOCTEXT("FailedCreateOutputDir", "Failed to create output directory \"{0}\""), FText::FromString(OutPath));
+			UE_LOG(LogTmvMedia, Error, TEXT("TranscodeMuxer::OpenStream: %s."), *Message.ToString());
+			UTmvMediaTranscodeJob::SafeReportError(InParentJob, Message);
 			return INDEX_NONE;
 		}
 	}
@@ -58,8 +64,9 @@ void UTmvMediaTranscodeMuxer::ReceiveAccessUnit(
 {
 	if (!InAccessUnit.IsValid())
 	{
-		UE_LOG(LogTmvMedia, Error, TEXT("ReceiveAccessUnit: Null access unit."));
-		UTmvMediaTranscodeJob::SafeReportError(InParentJob);
+		const FText Message = LOCTEXT("InvalidAccessUnit", "Null access unit");
+		UE_LOG(LogTmvMedia, Error, TEXT("TranscodeMuxer::ReceiveAccessUnit: %s."), *Message.ToString());
+		UTmvMediaTranscodeJob::SafeReportError(InParentJob, Message);
 		return; 
 	}
 
@@ -74,8 +81,9 @@ void UTmvMediaTranscodeMuxer::ReceiveAccessUnit(
 		FScopeLock Lock(&StreamsCS);
 		if (!Streams.IsValidIndex(InStreamId))
 		{
-			UE_LOG(LogTmvMedia, Error, TEXT("ReceiveAccessUnit: Invalid stream id."));
-			UTmvMediaTranscodeJob::SafeReportError(InParentJob);
+			const FText Message = LOCTEXT("InvalidStreamId", "Invalid stream id");
+			UE_LOG(LogTmvMedia, Error, TEXT("TranscodeMuxer::ReceiveAccessUnit: %s."), *Message.ToString());
+			UTmvMediaTranscodeJob::SafeReportError(InParentJob, Message);
 			return;
 		}
 		StreamInfo = Streams[InStreamId];
@@ -91,7 +99,7 @@ void UTmvMediaTranscodeMuxer::ReceiveAccessUnit(
 	// we can just write the file immediately.
 	// For a container muxer, we need to keep the access units and re-order them. They can be
 	// received out of order because of multiple worker threads. We would need to know the number of worker threads.
-	const FString OutPath = InParentJob->Settings.OutputPath.Path;
+	const FString OutPath = InParentJob->Settings.GetAbsoluteOutputPath();
 	const FString FileName = FString::Printf(TEXT("%s%05d.%s"), *StreamInfo.StreamName, InTimeInfo.FrameIndex, *StreamInfo.Extension);
 	const FString FullFileName = FPaths::Combine(OutPath, FPaths::MakeValidFileName(FileName));
 
@@ -100,13 +108,17 @@ void UTmvMediaTranscodeMuxer::ReceiveAccessUnit(
 	{
 		if (!File->Write(InAccessUnit->GetData(), InAccessUnit->Num()))
 		{
-			UE_LOG(LogTmvMedia, Error, TEXT("ReceiveAccessUnit: Failed to write access unit to file %s."), *FullFileName);
-			UTmvMediaTranscodeJob::SafeReportError(InParentJob);
+			const FText Message = FText::Format(LOCTEXT("WriteFailed", "Failed to write access unit to file \"{0}\""), FText::FromString(FullFileName));
+			UE_LOG(LogTmvMedia, Error, TEXT("TranscodeMuxer::ReceiveAccessUnit: %s."), *Message.ToString());
+			UTmvMediaTranscodeJob::SafeReportError(InParentJob, Message);
 		}
 	}
 	else
 	{
-		UE_LOG(LogTmvMedia, Error, TEXT("ReceiveAccessUnit: Failed to open output file for writing: %s"), *FullFileName);
-		UTmvMediaTranscodeJob::SafeReportError(InParentJob);
+		const FText Message = FText::Format(LOCTEXT("OpenFailed", "Failed to open output file for writing: \"{0}\""), FText::FromString(FullFileName));
+		UE_LOG(LogTmvMedia, Error, TEXT("TranscodeMuxer::ReceiveAccessUnit: %s."), *Message.ToString());
+		UTmvMediaTranscodeJob::SafeReportError(InParentJob, Message);
 	}
-}
\ No newline at end of file
+}
+
+#undef LOCTEXT_NAMESPACE
\ No newline at end of file
diff --git a/Engine/Plugins/Media/TmvMedia/Source/TmvMedia/Private/Transcoder/TmvMediaTranscodeSerialization.cpp b/Engine/Plugins/Media/TmvMedia/Source/TmvMedia/Private/Transcoder/TmvMediaTranscodeSerialization.cpp
--- a/Engine/Plugins/Media/TmvMedia/Source/TmvMedia/Private/Transcoder/TmvMediaTranscodeSerialization.cpp
+++ b/Engine/Plugins/Media/TmvMedia/Source/TmvMedia/Private/Transcoder/TmvMediaTranscodeSerialization.cpp
@@ -28,7 +28,7 @@ namespace UE::TmvMedia::TranscodeSerialization
 
 	/**
 	 * Try to get the field named FieldName as a string. Returns empty string if it doesn't exist or cannot be converted.
- 	 * (no assert/ensure/crash version)
+	  * (no assert/ensure/crash version)
 	 */
 	FString SafeGetStringField(const FJsonObject& InJsonRoot, FStringView InFieldName)
 	{
@@ -234,4 +234,95 @@ namespace UE::TmvMedia::TranscodeSerialization
 		UE_LOG(LogTmvMedia, Verbose, TEXT("Transcode list successfully loaded from json."));
 		return true;
 	}
+
+	bool SerializeTranscodeListItemsToJson(
+		FArchive& InArchive,
+		const TArray<FTmvMediaTranscodeListItem>& InTranscodeListItems)
+	{
+		using namespace SerializationUtils;
+		bool bSerializeError = false;
+
+		TArray<TSharedPtr<FJsonValue>> ItemEntries;
+		ItemEntries.Reserve(InTranscodeListItems.Num());
+
+		for (const FTmvMediaTranscodeListItem& Item : InTranscodeListItems)
+		{
+			if (TSharedPtr<FJsonObject> ItemJson = SerializeToJson(FTmvMediaTranscodeListItem::StaticStruct(), &Item))
+			{
+				ItemEntries.Add(MakeShared<FJsonValueObject>(ItemJson));	
+			}
+			else
+			{
+				UE_LOG(LogTmvMedia, Error, TEXT("Failed to serialize Transcode Item to json"));
+				bSerializeError = true;
+			}
+		}
+
+		const TSharedRef<FJsonObject> TranscodeListItemsJson = MakeShared<FJsonObject>();
+		TranscodeListItemsJson->SetArrayField(TEXT("Items"), ItemEntries);
+
+		const TSharedRef<TJsonWriter<>> Writer = TJsonWriterFactory<>::Create(&InArchive);
+		if (!FJsonSerializer::Serialize(TranscodeListItemsJson, Writer))
+		{
+			UE_LOG(LogTmvMedia, Error, TEXT("Serialize Transcode Items: Failed to write json string"));
+			bSerializeError = true;
+		}
+		return !bSerializeError;
+	}
+
+	bool DeserializeTranscodeListItemsFromJson(
+		FArchive& InArchive,
+		TArray<FTmvMediaTranscodeListItem>& OutTranscodeListItems)
+	{
+		using namespace SerializationUtils;
+
+		TSharedPtr<FJsonObject> JsonRoot = MakeShared<FJsonObject>();
+		const TSharedRef<TJsonReader<TCHAR>> JsonReader = TJsonReaderFactory<TCHAR>::Create(&InArchive);
+		if (!FJsonSerializer::Deserialize(JsonReader, JsonRoot) || !JsonRoot.IsValid())
+		{
+			UE_LOG(LogTmvMedia, Error, TEXT("Couldn't parse json data."));
+			return false;
+		}
+
+		const TArray<TSharedPtr<FJsonValue>>* ItemEntries;
+		if (!JsonRoot->TryGetArrayField(TEXT("Items"), ItemEntries))
+		{
+			UE_LOG(LogTmvMedia, Error, TEXT("Missing \"Items\" entry field in pasted text"));
+			return false;
+		}
+
+		OutTranscodeListItems.Reset();
+		OutTranscodeListItems.Reserve(ItemEntries->Num());
+	
+		for (const TSharedPtr<FJsonValue>& ItemEntry : *ItemEntries)
+		{
+			if (!ItemEntry.IsValid() || ItemEntry->Type != EJson::Object)
+			{
+				UE_LOG(LogTmvMedia, Warning, TEXT("Invalid item entry. Not an object"));
+				continue;
+			}
+
+			if (const TSharedPtr<FJsonObject>& ItemObject = ItemEntry->AsObject())
+			{
+				FTmvMediaTranscodeListItem Item;
+				FText ErrorMessage;
+
+				if (DeserializeFromJson(ItemObject.ToSharedRef(), FTmvMediaTranscodeListItem::StaticStruct(), &Item, ErrorMessage))
+				{
+					OutTranscodeListItems.Emplace(MoveTemp(Item));
+				}
+				else
+				{
+					UE_LOG(LogTmvMedia, Warning, TEXT("Unable to convert Item Entry Json Object to Transcode List Item Struct."));
+				}
+			}
+			else
+			{
+				UE_LOG(LogTmvMedia, Warning, TEXT("Item Entry is not a Json Object."));
+			}
+		}
+
+		UE_LOG(LogTmvMedia, Verbose, TEXT("Transcode list successfully loaded from json."));
+		return true;
+	}
 }
diff --git a/Engine/Plugins/Media/TmvMedia/Source/TmvMedia/Private/Utils/TmvMediaFrameUtils.cpp b/Engine/Plugins/Media/TmvMedia/Source/TmvMedia/Private/Utils/TmvMediaFrameUtils.cpp
--- a/Engine/Plugins/Media/TmvMedia/Source/TmvMedia/Private/Utils/TmvMediaFrameUtils.cpp
+++ b/Engine/Plugins/Media/TmvMedia/Source/TmvMedia/Private/Utils/TmvMediaFrameUtils.cpp
@@ -658,7 +658,7 @@ namespace UE::TmvMedia::FrameUtils::Pfm
 
 		uint8 BitDepth = MipInfo.Planes[0].BitDepth;
 		
-		for (const auto& PlaneInfo : MipInfo.Planes)
+		for (const FTmvMediaFramePlaneInfo& PlaneInfo : MipInfo.Planes)
 		{
 			// Component interleaving is not supported (yet?).
 			if (PlaneInfo.ComponentLayout == ETmvMediaFrameComponentLayout::Interleaved)
diff --git a/Engine/Plugins/Media/TmvMedia/Source/TmvMedia/Private/Utils/TmvMediaPathUtils.cpp b/Engine/Plugins/Media/TmvMedia/Source/TmvMedia/Private/Utils/TmvMediaPathUtils.cpp
new file mode 100644
--- /dev/null
+++ b/Engine/Plugins/Media/TmvMedia/Source/TmvMedia/Private/Utils/TmvMediaPathUtils.cpp
@@ -0,0 +1,225 @@
+// Copyright Epic Games, Inc. All Rights Reserved.
+
+#include "Utils/TmvMediaPathUtils.h"
+
+#include "Misc/Optional.h"
+#include "Misc/Paths.h"
+#include "Templates/Function.h"
+
+namespace UE::TmvMedia::PathUtils
+{
+	namespace Private
+	{
+		/** Definition of a path "token". */
+		struct FPathToken
+		{
+			bool IsValid() const
+			{
+				return !Token.IsEmpty() && GetPathFunction.IsSet();
+			}
+			
+			/** Token string. Ex: "{project_dir}". **/
+			FString Token;
+
+			/** Get the string the token resolve as. */
+			TFunction<FString()> GetPathFunction;
+		};
+
+		/** PathToken object used to return invalid token in path token api below. */
+		static FPathToken InvalidPathToken;
+
+		/**
+		 * List of currently supported path tokens.
+		 * For now, supporting the same tokens as ImgMediaSource.
+		 * See UImgMediaSource::ExpandSequencePathTokens.
+		 */
+		static const TArray<FPathToken> PathTokens =
+		{
+			{ TEXT("{project_dir}"), [](){ return FPaths::ProjectDir();} },
+			{ TEXT("{engine_dir}"),[](){ return FPaths::EngineDir();} }
+		};
+
+		/**
+		 * Replaces existing tokens for their full paths.
+		 * Ex: "{project_dir}/Movies" -> "C:/MyProject/Movies"
+		 * 
+		 * @remark Because this expands the tokens as "full paths", we are only supporting the paths
+		 * starting with the tokens and not in arbitrary positions.
+		 */
+		FString ExpandPathTokensToFull(const FString& InPath)
+		{
+			for (const FPathToken& PathToken : PathTokens)
+			{
+				if (InPath.StartsWith(PathToken.Token) && PathToken.GetPathFunction.IsSet())
+				{
+					const FString BaseDirectory = FPaths::ConvertRelativePathToFull(PathToken.GetPathFunction());
+					FString ExpandedPath = FPaths::Combine(BaseDirectory, InPath.RightChop(PathToken.Token.Len()));
+					FPaths::NormalizeDirectoryName(ExpandedPath); // Removes potential double slashes because of Combine.
+					return ExpandedPath;
+				}
+			}
+			return InPath;
+		}
+
+		/**
+		 * Finds the token this path string begins with.
+		 * @returns found token or invalid token if not found. Validate token with IsValid() before use.
+		 */
+		const FPathToken& FindFirstPathToken(const FString& InPath)
+		{
+			for (const FPathToken& PathToken : PathTokens)
+			{
+				if (InPath.StartsWith(PathToken.Token))
+				{
+					return PathToken;
+				}
+			}
+			return InvalidPathToken;
+		}
+
+		bool PathExists(const FString& InPath)
+		{
+			return FPaths::DirectoryExists(InPath) || FPaths::FileExists(InPath);
+		}
+
+		static TArray<TFunction<FString()>> BasePathFunctions =
+		{
+			[](){ return FPaths::ProjectContentDir();},
+			[](){ return FPaths::ProjectDir();},
+		};
+	}
+
+	FString EnsureStartWithDotSlash(const FString& InPath)
+	{
+		if (!InPath.StartsWith(TEXT("./")))
+		{
+			return FPaths::Combine(TEXT("."), InPath);
+		}
+		return InPath;
+	}
+
+	FString GetDirectoryFullPath(const FString& InSanitizedPath)
+	{
+		FString PathExpanded = Private::ExpandPathTokensToFull(InSanitizedPath);
+
+		// Relative paths are assumed to be under project directory, only support sanitized paths.
+		if (PathExpanded.StartsWith(TEXT("./")) || FPaths::IsRelative(PathExpanded))
+		{
+			const FString ProjectDir = FPaths::ConvertRelativePathToFull(FPaths::ProjectDir());
+			PathExpanded = FPaths::ConvertRelativePathToFull(ProjectDir, PathExpanded);
+		}
+		
+		// Check if it is an existing directory.
+		if (FPaths::DirectoryExists(PathExpanded))
+		{
+			return PathExpanded;
+		}
+
+		// If we reach here, we assume this is a filepath.
+		return FPaths::GetPath(PathExpanded);
+	}
+
+	FString GetSanitizedPath(const FString& InPath)
+	{
+		// Do nothing with empty strings.
+		if (InPath.IsEmpty())
+		{
+			return InPath;
+		}
+
+		FString SanitizedPath = InPath.TrimStartAndEnd().Replace(TEXT("\""), TEXT(""));
+		FPaths::NormalizeDirectoryName(SanitizedPath);
+
+		// If already a sanitized path, no need to process further.
+		// This will work even if the path doesn't exist.
+		if (SanitizedPath.StartsWith(TEXT("./")))
+		{
+			return SanitizedPath;
+		}
+
+		// Path strings using tokens also don't need further processing.
+		const Private::FPathToken& FirstToken = Private::FindFirstPathToken(SanitizedPath);
+		if (FirstToken.IsValid())
+		{
+			return SanitizedPath;
+		}
+
+		// Check if the path is under the current project
+		const FString ProjectDir = FPaths::ConvertRelativePathToFull(FPaths::ProjectDir());
+
+		// Test for existing relative path that might be under the project directory.
+		if (FPaths::IsRelative(SanitizedPath))
+		{
+			const FString ContentDir = FPaths::ConvertRelativePathToFull(FPaths::ProjectContentDir());
+			TArray<const FString*, TInlineAllocator<2>> BasePaths = { &ProjectDir, &ContentDir };
+
+			for (const FString* BasePath : BasePaths)
+			{
+				// This can return an absolute path that is not necessarily under the requested base path. Need to verify.
+				FString ConvertedPath = FPaths::ConvertRelativePathToFull(*BasePath, SanitizedPath);
+				
+				if (FPaths::IsUnderDirectory(ConvertedPath, *BasePath) && Private::PathExists(ConvertedPath))
+				{
+					// Make the path relative to project dir.
+					if (FPaths::MakePathRelativeTo(ConvertedPath, *ProjectDir))
+					{
+						return EnsureStartWithDotSlash(ConvertedPath);
+					}
+				}
+			}
+		}
+
+		// If the path is either absolute or not relative to project dir.
+		// We will attempt to make it absolute anyway (it may end up being under current process base dir if all else fails).
+		FString ConvertedPath = FPaths::ConvertRelativePathToFull(SanitizedPath);
+
+		// If the path is under the project directory, use relative path.
+		if (FPaths::IsUnderDirectory(ConvertedPath, ProjectDir))
+		{
+			// If so, make it relative to that.
+			if (FPaths::MakePathRelativeTo(ConvertedPath, *ProjectDir))
+			{
+				ConvertedPath = EnsureStartWithDotSlash(ConvertedPath);
+			}
+		}
+
+		return ConvertedPath;
+	}
+	
+	FString ConvertSanitizedPathToFull(const FString& InSanitizedPath)
+	{
+		if (InSanitizedPath.IsEmpty())
+		{
+			return FString();
+		}
+	
+		// Expand the common known path tokens.
+		FString PathExpanded = Private::ExpandPathTokensToFull(InSanitizedPath);
+
+		// If already a sanitized path, we consider it under project without further processing.
+		// This will work even if the path doesn't exist (for output paths for instance).
+		if (InSanitizedPath.StartsWith(TEXT("./")))
+		{
+			const FString BasePathFull = FPaths::ConvertRelativePathToFull(FPaths::ProjectDir());
+			return FPaths::ConvertRelativePathToFull(BasePathFull, PathExpanded);
+		}
+		
+		// If the path is relative, i.e. no tokens nor sanitized relative path.
+		if (FPaths::IsRelative(PathExpanded))
+		{
+			// Attempt to resolve against possible base paths. But this will only work if the resolved path exists.
+			for (const TFunction<FString()>& BasePathFunction : Private::BasePathFunctions)
+			{
+				const FString BasePathFull = FPaths::ConvertRelativePathToFull(BasePathFunction());
+				const FString PossiblePath = FPaths::ConvertRelativePathToFull(BasePathFull, PathExpanded);
+				if (FPaths::DirectoryExists(PossiblePath) || FPaths::FileExists(PossiblePath))
+				{
+					return PossiblePath;
+				}
+			}
+		}
+
+		// Final fallback: try to make it absolute (in case it was still relative to process base dir).
+		return FPaths::ConvertRelativePathToFull(PathExpanded);
+	}
+}
\ No newline at end of file
diff --git a/Engine/Plugins/Media/TmvMedia/Source/TmvMedia/Private/Utils/TmvMediaSerializationUtils.cpp b/Engine/Plugins/Media/TmvMedia/Source/TmvMedia/Private/Utils/TmvMediaSerializationUtils.cpp
--- a/Engine/Plugins/Media/TmvMedia/Source/TmvMedia/Private/Utils/TmvMediaSerializationUtils.cpp
+++ b/Engine/Plugins/Media/TmvMedia/Source/TmvMedia/Private/Utils/TmvMediaSerializationUtils.cpp
@@ -164,4 +164,23 @@ namespace UE::TmvMedia::SerializationUtils
 		constexpr bool bStrictMode = false;
 		return FJsonObjectConverter::JsonObjectToUStruct(InJsonObject, InStruct, InNativeObject, CheckFlags, SkipFlags, bStrictMode, &OutErrorMessage, &CustomCB);
 	}
+	
+	FString BytesToString(const TArray<uint8>& InValueAsBytes)
+	{
+		// Reinterpret as TCHAR array. FJsonStructSerializerBackend uses UCS2CHAR which is compatible with TCHAR.
+		const TCHAR* ValueAsChars = reinterpret_cast<const TCHAR*>(InValueAsBytes.GetData());
+		int32 SizeInChar = InValueAsBytes.Num() / sizeof(TCHAR);
+	
+		// Exclude terminating character, if any.
+		// Remark: From observation, the arrays from FJsonStructSerializerBackend don't seem to have a terminating character, but we check anyway.
+		if (SizeInChar > 0 && ValueAsChars[SizeInChar-1] == 0)
+		{
+			--SizeInChar; // Exclude terminating character.
+		}
+	
+		FString OutValueAsString;
+		OutValueAsString.Reset(SizeInChar);
+		OutValueAsString.AppendChars(ValueAsChars, SizeInChar);
+		return OutValueAsString;
+	}
 }
\ No newline at end of file
diff --git a/Engine/Plugins/Media/TmvMedia/Source/TmvMedia/Public/Transcoder/TmvMediaTranscodeJob.h b/Engine/Plugins/Media/TmvMedia/Source/TmvMedia/Public/Transcoder/TmvMediaTranscodeJob.h
--- a/Engine/Plugins/Media/TmvMedia/Source/TmvMedia/Public/Transcoder/TmvMediaTranscodeJob.h
+++ b/Engine/Plugins/Media/TmvMedia/Source/TmvMedia/Public/Transcoder/TmvMediaTranscodeJob.h
@@ -62,6 +62,9 @@ struct FTmvMediaTranscodeJobSettings
 	/** Returns true if the output path is set. */
 	UE_API bool IsOutputPathSet() const;
 
+	/** Returns the absolute output directory path. */ 
+	UE_API FString GetAbsoluteOutputPath() const;
+
 	/** Select the input source. */
 	UPROPERTY(EditAnywhere, Category = "Input Sequence", meta = (SegmentedDisplay))
 	ETmvMediaTranscodeInputSource InputSource = ETmvMediaTranscodeInputSource::File;
@@ -288,14 +291,16 @@ public:
 	 * Called by the processing stages in case a fatal error requires the job to stop.
 	 * Checks if the job is valid before reporting the error.
 	 */
-	static void SafeReportError(UTmvMediaTranscodeJob* InJob)
+	static void SafeReportError(UTmvMediaTranscodeJob* InJob, const FText& InErrorMessage)
 	{
-		// TODO: improve error reporting (add a message).
-		// For now stages causing the error should log something at least.
-		// Could use the Notification for error reporting message
 		if (IsValid(InJob))
 		{
 			InJob->bIsError = true;
+
+			if (!InErrorMessage.IsEmpty() && InJob->Notification.IsValid())
+			{
+				InJob->Notification->SetText(InErrorMessage);
+			}
 		}
 	}
 
diff --git a/Engine/Plugins/Media/TmvMedia/Source/TmvMedia/Public/Transcoder/TmvMediaTranscodeList.h b/Engine/Plugins/Media/TmvMedia/Source/TmvMedia/Public/Transcoder/TmvMediaTranscodeList.h
--- a/Engine/Plugins/Media/TmvMedia/Source/TmvMedia/Public/Transcoder/TmvMediaTranscodeList.h
+++ b/Engine/Plugins/Media/TmvMedia/Source/TmvMedia/Public/Transcoder/TmvMediaTranscodeList.h
@@ -49,6 +49,8 @@ enum class ETmvMediaTranscodeListItemEventType
 	ItemsAdded,
 	/** Item(s) removed from the list. */ 
 	ItemsRemoved,
+	/** Item(s) reordered in the list. */ 
+	ItemsReordered,
 	/** Item(s) properties modified. */
 	ItemsModified,
 };
@@ -110,6 +112,13 @@ public:
 	/** Remove items at given indices. */
 	UE_API void RemoveItems(TConstArrayView<int32> InItemIndices);
 
+	/**
+	 * Reorder an item from the given source index to the destination index.
+	 * @remark This function will broadcast the FOnItemEvent.
+	 * @return true if the operation completed, false if the indices where not valid.
+	 */
+	UE_API bool ReorderItems(int32 InFromIndex, int32 InToIndex);
+
 	/** Delegate for item editing events. */
 	DECLARE_MULTICAST_DELEGATE_TwoParams(FOnItemEvent, const UTmvMediaTranscodeList* /*List*/, const FTmvMediaTranscodeListItemEventArgs&);
 	FOnItemEvent& GetOnItemEvent()
diff --git a/Engine/Plugins/Media/TmvMedia/Source/TmvMediaEditor/Private/Customizations/TmvMediaDirectoryPathCustomization.cpp b/Engine/Plugins/Media/TmvMedia/Source/TmvMediaEditor/Private/Customizations/TmvMediaDirectoryPathCustomization.cpp
new file mode 100644
--- /dev/null
+++ b/Engine/Plugins/Media/TmvMedia/Source/TmvMediaEditor/Private/Customizations/TmvMediaDirectoryPathCustomization.cpp
@@ -0,0 +1,107 @@
+// Copyright Epic Games, Inc. All Rights Reserved.
+
+#include "TmvMediaDirectoryPathCustomization.h"
+
+#include "DetailWidgetRow.h"
+#include "IDetailChildrenBuilder.h"
+#include "Input/DragAndDrop.h"
+#include "SDropTarget.h"
+#include "Styling/AppStyle.h"
+#include "Utils/TmvMediaPathUtils.h"
+#include "Widgets/STmvMediaDirectoryPicker.h"
+
+#define LOCTEXT_NAMESPACE "FTmvMediaDirectoryPathCustomization"
+
+/* IPropertyTypeCustomization interface
+ *****************************************************************************/
+
+TSharedRef<IPropertyTypeCustomization> FTmvMediaDirectoryPathCustomization::MakeInstance()
+{
+	return MakeShared<FTmvMediaDirectoryPathCustomization>();
+}
+
+void FTmvMediaDirectoryPathCustomization::CustomizeHeader(TSharedRef<IPropertyHandle> InPropertyHandle,
+	FDetailWidgetRow& InHeaderRow, IPropertyTypeCustomizationUtils& InCustomizationUtils)
+{
+}
+
+void FTmvMediaDirectoryPathCustomization::CustomizeChildren(TSharedRef<IPropertyHandle> InPropertyHandle,
+	IDetailChildrenBuilder& InChildBuilder, IPropertyTypeCustomizationUtils& InCustomizationUtils)
+{
+	PathStringProperty = InPropertyHandle->GetChildHandle("Path");
+	
+	// Add Directory picker.
+	InChildBuilder.AddCustomRow(LOCTEXT("DirectoryPathLabel", "Path"))
+		.NameContent()
+		[
+			InPropertyHandle->CreatePropertyNameWidget()
+		]
+		.ValueContent()
+		.MaxDesiredWidth(0.0f)
+		.MinDesiredWidth(125.0f)
+		[
+			SNew(SDropTarget)	// Support drag and drop of external files.
+			.OnAllowDrop(this, &FTmvMediaDirectoryPathCustomization::HandleVerifyDrag)
+			.OnDropped(this, &FTmvMediaDirectoryPathCustomization::HandleDropEvent)
+			[
+				SNew(STmvMediaDirectoryPicker)
+				.BrowseButtonImage(FAppStyle::GetBrush("PropertyWindow.Button_Ellipsis"))
+				.BrowseButtonStyle(FAppStyle::Get(), "HoverHintOnly")
+				.BrowseButtonToolTip(LOCTEXT("DirectoryButtonToolTipText", "Choose a directory from this computer"))
+				.BrowseDirectory_Lambda([this]() -> FString
+				{
+					FString DirectoryPath = GetDirectoryPath();
+					return !DirectoryPath.IsEmpty() ?
+						UE::TmvMedia::PathUtils::ConvertSanitizedPathToFull(DirectoryPath) : FPaths::ConvertRelativePathToFull(FPaths::ProjectDir());
+				})
+				.BrowseTitle(LOCTEXT("BrowseTitle", "Directory picker..."))
+				.DirectoryPath(this, &FTmvMediaDirectoryPathCustomization::GetDirectoryPath)	
+				.OnDirectoryChanged(this, &FTmvMediaDirectoryPathCustomization::HandleDirectoryPickerDirectoryChanged)
+			]
+		];
+}
+
+FString FTmvMediaDirectoryPathCustomization::GetDirectoryPath() const
+{
+	FString FilePath;
+	if (PathStringProperty)
+	{
+		PathStringProperty->GetValue(FilePath);
+	}
+	return FilePath;
+}
+
+void FTmvMediaDirectoryPathCustomization::HandleDirectoryPickerDirectoryChanged(const FString& InPickedDirectory)
+{
+	if (PathStringProperty)
+	{
+		PathStringProperty->SetValue(UE::TmvMedia::PathUtils::GetSanitizedPath(InPickedDirectory));
+	}
+}
+
+bool FTmvMediaDirectoryPathCustomization::HandleVerifyDrag(TSharedPtr<FDragDropOperation> InDragOperation)
+{
+	if (InDragOperation && InDragOperation->IsOfType<FExternalDragOperation>())
+	{
+		TSharedPtr<FExternalDragOperation> ExternalDragOperation = StaticCastSharedPtr<FExternalDragOperation>(InDragOperation);
+		return ExternalDragOperation && ExternalDragOperation->HasFiles();
+	}
+	return false;
+}
+
+FReply FTmvMediaDirectoryPathCustomization::HandleDropEvent(const FGeometry& InGeometry, const FDragDropEvent& InDropEvent)
+{
+	const TSharedPtr<FDragDropOperation> DropOperation = InDropEvent.GetOperation();
+	if (DropOperation && DropOperation->IsOfType<FExternalDragOperation>())
+	{
+		TSharedPtr<FExternalDragOperation> ExternalDragOperation = StaticCastSharedPtr<FExternalDragOperation>(DropOperation);
+		if (ExternalDragOperation->HasFiles())
+		{
+			const FString SanitizedPath = UE::TmvMedia::PathUtils::GetSanitizedPath(ExternalDragOperation->GetFiles()[0]);
+			HandleDirectoryPickerDirectoryChanged(UE::TmvMedia::PathUtils::GetDirectoryFullPath(SanitizedPath));
+			return FReply::Handled();
+		}
+	}
+	return FReply::Unhandled();
+}
+#undef LOCTEXT_NAMESPACE
diff --git a/Engine/Plugins/Media/TmvMedia/Source/TmvMediaEditor/Private/Customizations/TmvMediaDirectoryPathCustomization.h b/Engine/Plugins/Media/TmvMedia/Source/TmvMediaEditor/Private/Customizations/TmvMediaDirectoryPathCustomization.h
new file mode 100644
--- /dev/null
+++ b/Engine/Plugins/Media/TmvMedia/Source/TmvMediaEditor/Private/Customizations/TmvMediaDirectoryPathCustomization.h
@@ -0,0 +1,55 @@
+// Copyright Epic Games, Inc. All Rights Reserved.
+
+#pragma once
+
+#include "IPropertyTypeCustomization.h"
+#include "Input/Reply.h"
+
+class FDragDropOperation;
+class IPropertyHandle;
+
+/**
+ * Implements a customization for the FDirectoryPath class for Tmv transcoder.
+ *
+ * Details:
+ * - The picked path is either made absolute if outside the project or relative to the project.
+ * - Implements drag and drop of files from the explorer.
+ */
+class FTmvMediaDirectoryPathCustomization : public IPropertyTypeCustomization
+{
+public:
+	/** Make an instance of this property type customization. */
+	static TSharedRef<IPropertyTypeCustomization> MakeInstance();
+
+	//~ Begin IPropertyTypeCustomization
+	virtual void CustomizeHeader(TSharedRef<IPropertyHandle> InPropertyHandle, FDetailWidgetRow& InHeaderRow,
+		IPropertyTypeCustomizationUtils& InCustomizationUtils) override;
+	virtual void CustomizeChildren(TSharedRef<IPropertyHandle> InPropertyHandle, IDetailChildrenBuilder& InChildBuilder,
+		IPropertyTypeCustomizationUtils& InCustomizationUtils) override;
+	//~ End IPropertyTypeCustomization
+
+private:
+	/**
+	 * Gets the directory path from the property handle.
+	 */
+	FString GetDirectoryPath() const;
+
+	/**
+	 * Called when a directory is selected.
+	 */
+	void HandleDirectoryPickerDirectoryChanged(const FString& InPickedDirectory);
+
+	/**
+	 * Called when a drag operation enters the widget. 
+	 */
+	bool HandleVerifyDrag(TSharedPtr<FDragDropOperation> InDragOperation);
+
+	/**
+	 * Called on a drag drop event.
+	 */
+	FReply HandleDropEvent(const FGeometry& InGeometry, const FDragDropEvent& InDropEvent);
+
+	/** Pointer to the string that will be set when changing the path */
+	TSharedPtr<IPropertyHandle> PathStringProperty;
+};
+
diff --git a/Engine/Plugins/Media/TmvMedia/Source/TmvMediaEditor/Private/Customizations/TmvMediaFilePathCustomization.cpp b/Engine/Plugins/Media/TmvMedia/Source/TmvMediaEditor/Private/Customizations/TmvMediaFilePathCustomization.cpp
--- a/Engine/Plugins/Media/TmvMedia/Source/TmvMediaEditor/Private/Customizations/TmvMediaFilePathCustomization.cpp
+++ b/Engine/Plugins/Media/TmvMedia/Source/TmvMediaEditor/Private/Customizations/TmvMediaFilePathCustomization.cpp
@@ -5,9 +5,9 @@
 #include "DetailWidgetRow.h"
 #include "IDetailChildrenBuilder.h"
 #include "Input/DragAndDrop.h"
-#include "Misc/Paths.h"
 #include "SDropTarget.h"
 #include "Styling/AppStyle.h"
+#include "Utils/TmvMediaPathUtils.h"
 #include "Widgets/Input/SFilePathPicker.h"
 
 #define LOCTEXT_NAMESPACE "FTmvMediaFilePathCustomization"
@@ -53,7 +53,7 @@ void FTmvMediaFilePathCustomization::CustomizeChildren(TSharedRef<IPropertyHandl
 				.BrowseButtonToolTip(LOCTEXT("FileButtonToolTipText", "Choose a file from this computer"))
 				.BrowseDirectory_Lambda([this]()->FString
 				{
-					return FPaths::GetPath(HandleFilePathPickerFilePath());
+					return UE::TmvMedia::PathUtils::GetDirectoryFullPath(HandleFilePathPickerFilePath());
 				})
 				.BrowseTitle(LOCTEXT("PropertyEditorTitle", "File picker..."))
 				.FilePath(this, &FTmvMediaFilePathCustomization::HandleFilePathPickerFilePath)
@@ -75,24 +75,14 @@ FString FTmvMediaFilePathCustomization::HandleFilePathPickerFilePath() const
 
 void FTmvMediaFilePathCustomization::HandleFilePathPickerPathPicked(const FString& InPickedPath)
 {
+	using namespace UE::TmvMedia::PathUtils;
+
 	if (!PathStringProperty)
 	{
 		return;
 	}
 
-	// The picker usually returns a path relative to executable. Absolute path is preferable.
-	FString ConfigFilepath = FPaths::ConvertRelativePathToFull(InPickedPath);
-
-	// Check if the path is under the current project
-	const FString ProjectDir = FPaths::ConvertRelativePathToFull(FPaths::ProjectDir());
-
-	if (FPaths::IsUnderDirectory(ConfigFilepath, ProjectDir))
-	{
-		// If so, make it relative to that.
-		FPaths::MakePathRelativeTo(ConfigFilepath, *ProjectDir);
-	}
-
-	PathStringProperty->SetValue(ConfigFilepath);
+	PathStringProperty->SetValue(GetSanitizedPath(InPickedPath));
 }
 
 bool FTmvMediaFilePathCustomization::HandleVerifyDrag(TSharedPtr<FDragDropOperation> InDragOperation)
diff --git a/Engine/Plugins/Media/TmvMedia/Source/TmvMediaEditor/Private/Widgets/STmvMediaDirectoryPicker.cpp b/Engine/Plugins/Media/TmvMedia/Source/TmvMediaEditor/Private/Widgets/STmvMediaDirectoryPicker.cpp
new file mode 100644
--- /dev/null
+++ b/Engine/Plugins/Media/TmvMedia/Source/TmvMediaEditor/Private/Widgets/STmvMediaDirectoryPicker.cpp
@@ -0,0 +1,111 @@
+// Copyright Epic Games, Inc. All Rights Reserved.
+
+#include "STmvMediaDirectoryPicker.h"
+
+#include "DesktopPlatformModule.h"
+#include "Framework/Application/SlateApplication.h"
+#include "Widgets/Images/SImage.h"
+#include "Widgets/Input/SButton.h"
+#include "Widgets/Input/SEditableTextBox.h"
+
+#define LOCTEXT_NAMESPACE "STmvMediaDirectoryPicker"
+
+void STmvMediaDirectoryPicker::Construct( const FArguments& InArgs )
+{
+	BrowseTitle = InArgs._BrowseTitle;
+	BrowseDirectory = InArgs._BrowseDirectory;
+	DirectoryPath = InArgs._DirectoryPath;
+	OnDirectoryChanged = InArgs._OnDirectoryChanged;
+	
+	TSharedPtr<SButton> OpenButton;
+	ChildSlot
+	[
+		SNew(SHorizontalBox)
+
+		+ SHorizontalBox::Slot()
+		.FillWidth(1.0f)
+		[
+			SAssignNew(DirectoryTextBox, SEditableTextBox)
+			.Text(this, &STmvMediaDirectoryPicker::GetDirectoryTextBoxText)
+			.Font(InArgs._Font)
+			.SelectAllTextWhenFocused(true)
+			.ClearKeyboardFocusOnCommit(false)
+			.OnTextCommitted(this, &STmvMediaDirectoryPicker::OnDirectoryTextCommited)
+			.SelectAllTextOnCommit(false)
+			.IsReadOnly(InArgs._IsReadOnly)
+		]
+
+		+ SHorizontalBox::Slot()
+		.AutoWidth()
+		.Padding(4.0f, 0.0f, 0.0f, 0.0f)
+		.VAlign(VAlign_Center)
+		[
+			SAssignNew(OpenButton, SButton)
+			.ButtonStyle(InArgs._BrowseButtonStyle)
+			.ToolTipText(InArgs._BrowseButtonToolTip)
+			.OnClicked(this, &STmvMediaDirectoryPicker::OnBrowseButtonClicked)
+			.ContentPadding(2.0f)
+			.ForegroundColor(FSlateColor::UseForeground())
+			.IsFocusable(false)
+			[
+				SNew(SImage)
+				.Image(InArgs._BrowseButtonImage)
+				.ColorAndOpacity(FSlateColor::UseForeground())
+			]
+		]
+	];
+
+	OpenButton->SetEnabled(InArgs._IsEnabled);
+}
+
+void STmvMediaDirectoryPicker::OnDirectoryTextCommited(const FText& InText, ETextCommit::Type InCommitType)
+{
+	OnDirectoryChanged.ExecuteIfBound(InText.ToString());
+}
+
+FText STmvMediaDirectoryPicker::GetDirectoryTextBoxText() const
+{
+	return FText::FromString(DirectoryPath.Get());
+}
+
+FReply STmvMediaDirectoryPicker::OnBrowseButtonClicked()
+{
+	const FString DefaultPath = BrowseDirectory.IsSet() ? BrowseDirectory.Get() : DirectoryPath.Get();
+	FString OutDirectory;
+
+	if ( OpenPlatformDirectoryPicker(OutDirectory, DefaultPath) )
+	{
+		OnDirectoryChanged.ExecuteIfBound(OutDirectory);
+	}
+
+	return FReply::Handled();
+}
+
+bool STmvMediaDirectoryPicker::OpenPlatformDirectoryPicker(FString& OutDirectory, const FString& InDefaultPath) const
+{
+	bool bFolderSelected = false;
+	if ( IDesktopPlatform* DesktopPlatform = FDesktopPlatformModule::Get() )
+	{
+		TSharedPtr<SWindow> ParentWindow = FSlateApplication::Get().FindWidgetWindow(AsShared());
+		void* ParentWindowHandle = (ParentWindow.IsValid() && ParentWindow->GetNativeWindow().IsValid())
+			? ParentWindow->GetNativeWindow()->GetOSWindowHandle()
+			: nullptr;
+
+		FString OutFolderName;
+		bFolderSelected = DesktopPlatform->OpenDirectoryDialog(
+			ParentWindowHandle,
+			BrowseTitle.Get().ToString(),
+			InDefaultPath,
+			OutFolderName
+		);
+
+		if ( bFolderSelected )
+		{
+			OutDirectory = OutFolderName;
+		}
+	}
+
+	return bFolderSelected;
+}
+
+#undef LOCTEXT_NAMESPACE
diff --git a/Engine/Plugins/Media/TmvMedia/Source/TmvMediaEditor/Private/Widgets/STmvMediaDirectoryPicker.h b/Engine/Plugins/Media/TmvMedia/Source/TmvMediaEditor/Private/Widgets/STmvMediaDirectoryPicker.h
new file mode 100644
--- /dev/null
+++ b/Engine/Plugins/Media/TmvMedia/Source/TmvMediaEditor/Private/Widgets/STmvMediaDirectoryPicker.h
@@ -0,0 +1,95 @@
+// Copyright Epic Games, Inc. All Rights Reserved.
+
+#pragma once
+
+#include "Fonts/SlateFontInfo.h"
+#include "Input/Reply.h"
+#include "Misc/Attribute.h"
+#include "Styling/ISlateStyle.h"
+#include "Styling/SlateWidgetStyleAsset.h"
+#include "Widgets/DeclarativeSyntaxSupport.h"
+#include "Widgets/SCompoundWidget.h"
+
+class SEditableTextBox;
+
+/**
+ * A directory path box, with a button for picking a new path.
+ * Reimplementation of SDirectoryPicker.
+ * Differences:
+ *  - Made the text box display a delegate to allow user class to intercept and sanitize the picked path.
+ *  - Made it match the look, feel and functionality of SFilePathPicker (font, style, title, etc).
+ */
+class STmvMediaDirectoryPicker : public SCompoundWidget
+{
+public:
+	/** Declares a delegate that is executed when a directory was picked. */
+	DECLARE_DELEGATE_OneParam(FOnDirectoryChanged, const FString& /*Directory*/);
+
+	SLATE_BEGIN_ARGS(STmvMediaDirectoryPicker)
+		: _BrowseButtonToolTip(NSLOCTEXT("STmvMediaDirectoryPicker", "BrowseButtonToolTip", "Choose a directory from this computer"))
+		, _Font()
+		, _IsReadOnly(false)
+	{ }
+
+		/** Browse button image resource. */
+		SLATE_ATTRIBUTE(const FSlateBrush*, BrowseButtonImage)
+
+		/** Browse button visual style. */
+		SLATE_STYLE_ARGUMENT(FButtonStyle, BrowseButtonStyle)
+
+		/** Browse button tool tip text. */
+		SLATE_ATTRIBUTE(FText, BrowseButtonToolTip)
+
+		/** The directory to browse by default */
+		SLATE_ATTRIBUTE(FString, BrowseDirectory)
+
+		/** Title for the browse dialog window. */
+		SLATE_ATTRIBUTE(FText, BrowseTitle)
+		
+		/** Holds the currently selected directory path. */
+		SLATE_ATTRIBUTE(FString, DirectoryPath)
+
+		/** Sets the font used for the path text box. */
+		SLATE_ATTRIBUTE(FSlateFontInfo, Font)
+
+		/** Whether the path text box can be modified by the user. */
+		SLATE_ATTRIBUTE(bool, IsReadOnly)
+
+		/** Called when a path has been picked or modified. */
+		SLATE_EVENT(FOnDirectoryChanged, OnDirectoryChanged)
+
+	SLATE_END_ARGS()
+
+public:
+
+	void Construct(const FArguments& InArgs);
+
+private:
+	/** DirectoryTextBox - OnTextCommitted event handler. */
+	void OnDirectoryTextCommited(const FText& InText, ETextCommit::Type InCommitType);
+
+	/** DirectoryTextBox - Text attribute handler. */
+	FText GetDirectoryTextBoxText() const;
+
+	/** Browse Button - OnClicked event handler. */
+	FReply OnBrowseButtonClicked();
+	
+	/** Utility function to open a directory picker dialog. */
+	bool OpenPlatformDirectoryPicker(FString& OutDirectory, const FString& InDefaultPath) const;
+
+private:
+	/** Holds the directory path to browse by default. */
+	TAttribute<FString> BrowseDirectory;
+
+	/** Holds the title for the browse dialog window. */
+	TAttribute<FText> BrowseTitle;
+	
+	/** Holds the currently selected directory path. */
+	TAttribute<FString> DirectoryPath;
+
+	/** Holds a delegate that is executed when a directory was picked. */
+	FOnDirectoryChanged OnDirectoryChanged;
+
+	/** Editable text box for the directory path. */
+	TSharedPtr<SEditableTextBox> DirectoryTextBox;
+};
\ No newline at end of file
diff --git a/Engine/Plugins/Media/TmvMedia/Source/TmvMediaEditor/Private/Widgets/STmvMediaTranscodeList.cpp b/Engine/Plugins/Media/TmvMedia/Source/TmvMediaEditor/Private/Widgets/STmvMediaTranscodeList.cpp
--- a/Engine/Plugins/Media/TmvMedia/Source/TmvMediaEditor/Private/Widgets/STmvMediaTranscodeList.cpp
+++ b/Engine/Plugins/Media/TmvMedia/Source/TmvMediaEditor/Private/Widgets/STmvMediaTranscodeList.cpp
@@ -2,13 +2,18 @@
 
 #include "STmvMediaTranscodeList.h"
 
+#include "Framework/Commands/GenericCommands.h"
+#include "HAL/PlatformApplicationMisc.h"
 #include "ScopedTransaction.h"
 #include "SlateOptMacros.h"
 #include "TmvMediaEditorLog.h"
 #include "TmvMediaTranscodeListCommands.h"
 #include "TmvMediaTranscodeListHandle.h"
+#include "DragAndDrop/DecoratedDragDropOp.h"
 #include "Transcoder/ITmvMediaTranscodeJobManager.h"
 #include "Transcoder/TmvMediaTranscodeList.h"
+#include "Transcoder/TmvMediaTranscodeSerialization.h"
+#include "Utils/TmvMediaSerializationUtils.h"
 #include "Widgets/Input/SButton.h"
 #include "Widgets/Layout/SWidgetSwitcher.h"
 #include "Widgets/Notifications/SProgressBar.h"
@@ -222,6 +227,34 @@ struct FTmvMediaTranscodeJobListItem : public TSharedFromThis<FTmvMediaTranscode
 
 	/** Cached the item guid to retrieve the corresponding job in the job manager. */
 	FGuid ItemId;
+
+	/** Keep a pointer to the name edit box to be able to rename with F2 */
+	TWeakPtr<SInlineEditableTextBlock> NameEditableBoxWeak;
+};
+
+/**
+ * This drag drop operation for one transcode job item in the list view.
+ * Using decorated drag and drop for clarity.
+ */
+class FTmvMediaTranscodeJobListItemDragDropOp : public FDecoratedDragDropOp
+{
+public:
+	DRAG_DROP_OPERATOR_TYPE(FTmvMediaTranscodeJobListItemDragDropOp, FDecoratedDragDropOp)
+
+	/** The template to create an instance */
+	TSharedPtr<FTmvMediaTranscodeJobListItem> ListItem;
+
+	/** Constructs the drag drop operation */
+	static TSharedRef<FTmvMediaTranscodeJobListItemDragDropOp> New(const TSharedPtr<FTmvMediaTranscodeJobListItem>& InListItem, FText InDragText)
+	{
+		TSharedRef<FTmvMediaTranscodeJobListItemDragDropOp> Operation = MakeShared<FTmvMediaTranscodeJobListItemDragDropOp>();
+		Operation->ListItem = InListItem;
+		Operation->DefaultHoverText = InDragText;
+		Operation->CurrentHoverText = InDragText;
+		Operation->Construct();
+
+		return Operation;
+	}
 };
 
 /**
@@ -243,9 +276,64 @@ public:
 	void Construct(const FArguments& InArgs, const TSharedRef<STableViewBase>& InOwnerTable)
 	{
 		Item = InArgs._Item;
-		FSuperRowType::FArguments SuperArgs = FSuperRowType::FArguments();
 		// See SQueueJobListRow::Construct for an example of how to connect Drag events.
-		FSuperRowType::Construct(SuperArgs, InOwnerTable);
+		FSuperRowType::Construct(
+			FSuperRowType::FArguments()
+			.OnDragDetected(this, &STmvMediaTranscodeJobListRow::OnDragDetected)
+			.OnCanAcceptDrop(this, &STmvMediaTranscodeJobListRow::OnCanAcceptDrop)
+			.OnAcceptDrop(this, &STmvMediaTranscodeJobListRow::OnAcceptDrop)
+			, InOwnerTable);
+	}
+	
+	FReply OnDragDetected(const FGeometry& InMyGeometry, const FPointerEvent& InPointerEvent)
+	{
+		if (Item.IsValid())
+		{
+			FText DefaultText = LOCTEXT("DefaultDragDropFormat", "Move 1 item(s)");
+			return FReply::Handled().BeginDragDrop(FTmvMediaTranscodeJobListItemDragDropOp::New(Item, DefaultText));
+		}
+		return FReply::Unhandled();
+	}
+
+	TOptional<EItemDropZone> OnCanAcceptDrop(const FDragDropEvent& InDragDropEvent, EItemDropZone InDropZone, TSharedPtr<FTmvMediaTranscodeJobListItem> InItem)
+	{
+		TSharedPtr<FTmvMediaTranscodeJobListItemDragDropOp> DragDropOp = InDragDropEvent.GetOperationAs<FTmvMediaTranscodeJobListItemDragDropOp>();
+		if (DragDropOp.IsValid())
+		{
+			if (InDropZone == EItemDropZone::OntoItem)
+			{
+				DragDropOp->CurrentIconBrush = FAppStyle::GetBrush(TEXT("Graph.ConnectorFeedback.Error"));
+			}
+			else
+			{
+				DragDropOp->CurrentIconBrush = FAppStyle::GetBrush(TEXT("Graph.ConnectorFeedback.Ok"));
+			}
+			return InDropZone;
+		}
+		return TOptional<EItemDropZone>();
+	}
+
+	FReply OnAcceptDrop(const FDragDropEvent& InDragDropEvent, EItemDropZone InDropZone, TSharedPtr<FTmvMediaTranscodeJobListItem> InItem)
+	{
+		TSharedPtr<FTmvMediaTranscodeJobListItemDragDropOp> DragDropOp = InDragDropEvent.GetOperationAs<FTmvMediaTranscodeJobListItemDragDropOp>();
+		if (!DragDropOp.IsValid() || !DragDropOp->ListItem.IsValid() || !Item.IsValid())
+		{
+			return FReply::Unhandled();	
+		}
+
+		UTmvMediaTranscodeList* TranscodeList = Item->TranscodeListWeak.Get();
+
+		if (!TranscodeList || TranscodeList != DragDropOp->ListItem->TranscodeListWeak.Get())
+		{
+			return FReply::Unhandled(); // Only support drag and drop from the same list for now.
+		}
+		
+		FScopedTransaction Transaction(LOCTEXT("ReorderJob_Transaction", "Reorder Job"));
+		TranscodeList->Modify();
+
+		const int32 IndexOffset = (InDropZone == EItemDropZone::BelowItem) ? 1 : 0;
+		return TranscodeList->ReorderItems(DragDropOp->ListItem->ItemIndex, Item->ItemIndex + IndexOffset)
+			? FReply::Handled() : FReply::Unhandled();
 	}
 
 	//~ Begin SMultiColumnTableRow
@@ -256,9 +344,7 @@ public:
 			return SNew(SBox)
 			.Padding(2.0f)
 			[
-				//SNew(STextBlock)
-				//.Text(Item.Get(), &FTmvMediaTranscodeJobListItem::GetJobName)
-				SNew(SInlineEditableTextBlock)
+				SAssignNew(Item->NameEditableBoxWeak, SInlineEditableTextBlock)
 				.Text(Item.Get(), &FTmvMediaTranscodeJobListItem::GetJobName)
 				.OnVerifyTextChanged(Item.Get(), &FTmvMediaTranscodeJobListItem::OnVerifyJobNameTextChanged)
 				.OnTextCommitted(Item.Get(), &FTmvMediaTranscodeJobListItem::OnJobNameTextCommitted)
@@ -436,8 +522,14 @@ void STmvMediaTranscodeList::Construct(const FArguments& InArgs)
 
 	using namespace UE::TmvMediaEditor::TranscodeList;
 
-	TSharedPtr<FUICommandList> CommandList = InArgs._CommandList;
-	
+	// Create a command list of this widget to add the transcode list specific commands.
+	CommandList = MakeShared<FUICommandList>();
+	BindCommands(CommandList);
+	if (InArgs._CommandList)
+	{
+		CommandList->Append(InArgs._CommandList.ToSharedRef());
+	}
+
 	ChildSlot
 	[
 		SNew(SVerticalBox)
@@ -459,8 +551,8 @@ void STmvMediaTranscodeList::Construct(const FArguments& InArgs)
 			.AutoWidth()
 			[
 				SNew(SButton)
-					.IsEnabled(this, &STmvMediaTranscodeList::CanDuplicateJob)
-					.OnClicked(this, &STmvMediaTranscodeList::OnDuplicateJob)
+					.IsEnabled(this, &STmvMediaTranscodeList::CanDuplicateSelectedJobs)
+					.OnClicked(this, &STmvMediaTranscodeList::OnDuplicateSelectedJobs)
 					.ForegroundColor(FSlateColor::UseStyle())
 					.Text(LOCTEXT("DuplicateJobButtonLabel", "Duplicate Job"))
 					.ToolTipText(LOCTEXT("DuplicateJobButtonToolTip", "Duplicate selected Job"))
@@ -469,8 +561,8 @@ void STmvMediaTranscodeList::Construct(const FArguments& InArgs)
 			.AutoWidth()
 			[
 				SNew(SButton)
-					.IsEnabled(this, &STmvMediaTranscodeList::CanRemoveJob)
-					.OnClicked(this, &STmvMediaTranscodeList::OnRemoveJob)
+					.IsEnabled(this, &STmvMediaTranscodeList::CanRemoveSelectedJobs)
+					.OnClicked(this, &STmvMediaTranscodeList::OnRemoveSelectedJobs)
 					.ForegroundColor(FSlateColor::UseStyle())
 					.Text(LOCTEXT("RemoveJobButtonLabel", "Remove Job(s)"))
 					.ToolTipText(LOCTEXT("RemoveJobButtonToolTip", "Remove selected job(s) from the list."))
@@ -515,6 +607,44 @@ void STmvMediaTranscodeList::Construct(const FArguments& InArgs)
 }
 END_SLATE_FUNCTION_BUILD_OPTIMIZATION
 
+FReply STmvMediaTranscodeList::OnKeyDown(const FGeometry& MyGeometry, const FKeyEvent& InKeyEvent)
+{
+	if (CommandList && CommandList->ProcessCommandBindings(InKeyEvent))
+	{
+		return FReply::Handled();
+	}
+	return SCompoundWidget::OnKeyDown(MyGeometry, InKeyEvent);	
+}
+
+void STmvMediaTranscodeList::BindCommands(const TSharedPtr<FUICommandList>& InCommandList)
+{
+	const FGenericCommands& GenericCommands = FGenericCommands::Get();
+
+	InCommandList->MapAction(GenericCommands.Rename,
+		FExecuteAction::CreateSP(this, &STmvMediaTranscodeList::RenameSelectedJob),
+		FCanExecuteAction::CreateSP(this, &STmvMediaTranscodeList::CanRenameSelectedJob));
+
+	InCommandList->MapAction(GenericCommands.Delete,
+		FExecuteAction::CreateSPLambda(this, [this](){ OnRemoveSelectedJobs(); }),
+		FCanExecuteAction::CreateSP(this, &STmvMediaTranscodeList::CanRemoveSelectedJobs));
+
+	InCommandList->MapAction(GenericCommands.Cut,
+		FExecuteAction::CreateSP(this, &STmvMediaTranscodeList::CutSelectedJobs),
+		FCanExecuteAction::CreateSP(this, &STmvMediaTranscodeList::CanCutSelectedJobs));
+
+	InCommandList->MapAction(GenericCommands.Copy,
+		FExecuteAction::CreateSP(this, &STmvMediaTranscodeList::CopySelectedJobs),
+		FCanExecuteAction::CreateSP(this, &STmvMediaTranscodeList::CanCopySelectedJobs));
+
+	InCommandList->MapAction(GenericCommands.Paste,
+		FExecuteAction::CreateSP(this, &STmvMediaTranscodeList::Paste),
+		FCanExecuteAction::CreateSP(this, &STmvMediaTranscodeList::CanPaste));
+
+	InCommandList->MapAction(GenericCommands.Duplicate,
+		FExecuteAction::CreateSPLambda(this, [this](){ OnDuplicateSelectedJobs(); }),
+		FCanExecuteAction::CreateSP(this, &STmvMediaTranscodeList::CanDuplicateSelectedJobs));
+}
+
 TSharedRef<ITableRow> STmvMediaTranscodeList::OnGenerateRow(TSharedPtr<FTmvMediaTranscodeJobListItem> InItem, const TSharedRef<STableViewBase>& InOwnerTable) const
 {
 	return SNew(STmvMediaTranscodeJobListRow, InOwnerTable)
@@ -555,6 +685,11 @@ void STmvMediaTranscodeList::OnSelectionChanged(TSharedPtr<FTmvMediaTranscodeJob
 	}
 }
 
+bool STmvMediaTranscodeList::IsSelectionValid() const
+{
+	return JobListView && JobListView->GetNumItemsSelected() > 0;
+}
+
 bool STmvMediaTranscodeList::CanAddJob() const
 {
 	return true;
@@ -571,12 +706,7 @@ FReply STmvMediaTranscodeList::OnAddJob()
 	return FReply::Handled();
 }
 
-bool STmvMediaTranscodeList::CanDuplicateJob() const
-{
-	return JobListView && JobListView->GetNumItemsSelected() > 0;
-}
-
-FReply STmvMediaTranscodeList::OnDuplicateJob()
+FReply STmvMediaTranscodeList::OnDuplicateSelectedJobs()
 {
 	if (UTmvMediaTranscodeList* List = ListHandle ? ListHandle->Get() : nullptr)
 	{
@@ -593,12 +723,7 @@ FReply STmvMediaTranscodeList::OnDuplicateJob()
 	return FReply::Handled();
 }
 
-bool STmvMediaTranscodeList::CanRemoveJob() const
-{
-	return JobListView && JobListView->GetNumItemsSelected() > 0;
-}
-
-FReply STmvMediaTranscodeList::OnRemoveJob()
+FReply STmvMediaTranscodeList::OnRemoveSelectedJobs()
 {
 	if (UTmvMediaTranscodeList* List = ListHandle ? ListHandle->Get() : nullptr)
 	{
@@ -608,7 +733,7 @@ FReply STmvMediaTranscodeList::OnRemoveJob()
 		TArray<int32> ItemsToRemove;
 		TArray<TSharedPtr<FTmvMediaTranscodeJobListItem>> SelectedItems = JobListView->GetSelectedItems();
 		ItemsToRemove.Reserve(SelectedItems.Num());
-		for (const auto& SelectedItem : SelectedItems)
+		for (const TSharedPtr<FTmvMediaTranscodeJobListItem>& SelectedItem : SelectedItems)
 		{
 			ItemsToRemove.Add(SelectedItem->ItemIndex);
 		}
@@ -618,6 +743,128 @@ FReply STmvMediaTranscodeList::OnRemoveJob()
 	return FReply::Handled();
 }
 
+bool STmvMediaTranscodeList::CanRenameSelectedJob() const
+{
+	return JobListView && JobListView->GetNumItemsSelected() == 1;
+}
+
+void STmvMediaTranscodeList::RenameSelectedJob()
+{
+	if (JobListView)
+	{
+		const TArray<TSharedPtr<FTmvMediaTranscodeJobListItem>> SelectedItems = JobListView->GetSelectedItems();
+		if (SelectedItems.Num() > 0)
+		{
+			if (TSharedPtr<SInlineEditableTextBlock> NameEditableBox = SelectedItems[0]->NameEditableBoxWeak.Pin())
+			{
+				NameEditableBox->EnterEditingMode();
+			}
+		}
+	}
+}
+
+void STmvMediaTranscodeList::CutSelectedJobs()
+{
+	CopySelectedJobs();
+	OnRemoveSelectedJobs();
+}
+
+namespace UE::TmvMediaEditor::TranscodeList
+{
+	static FString JobClipboardPrefix = TEXT("TmvMediaTranscodeItems");
+}
+
+void STmvMediaTranscodeList::CopySelectedJobs()
+{
+	UTmvMediaTranscodeList* List = ListHandle ? ListHandle->Get() : nullptr;
+	if (!List || !JobListView)
+	{
+		return;
+	}
+
+	TArray<TSharedPtr<FTmvMediaTranscodeJobListItem>> SelectedItems = JobListView->GetSelectedItems();
+	if (SelectedItems.IsEmpty())
+	{
+		UE_LOG(LogTmvMediaEditor, Warning, TEXT("No items selected."));
+		return;
+	}
+	
+	TArray<FTmvMediaTranscodeListItem> Items;
+	Items.Reserve(SelectedItems.Num());
+	for(const TSharedPtr<FTmvMediaTranscodeJobListItem>& SelectedItem : SelectedItems)
+	{
+		if (SelectedItem)
+		{
+			Items.Add(List->GetItem(SelectedItem->ItemIndex));
+		}
+	}
+
+	TArray<uint8> Memory;
+	FMemoryWriter MemoryWriter(Memory);
+	if (!UE::TmvMedia::TranscodeSerialization::SerializeTranscodeListItemsToJson(MemoryWriter, Items))
+	{
+		return;	// SerializeTranscodeListItemsToJson is logging errors already.
+	}
+
+	FString SerializedString = UE::TmvMedia::SerializationUtils::BytesToString(Memory);
+	// Add Prefix to quickly identify whether current clipboard is from Transcode items or not
+	SerializedString = *FString::Printf(TEXT("%s%s"), *UE::TmvMediaEditor::TranscodeList::JobClipboardPrefix, *SerializedString);
+	FPlatformApplicationMisc::ClipboardCopy(*SerializedString);
+}
+
+bool STmvMediaTranscodeList::CanPaste() const
+{
+	if (CanAddJob())
+	{
+		FString PastedText;
+		FPlatformApplicationMisc::ClipboardPaste(PastedText);
+		return PastedText.StartsWith(UE::TmvMediaEditor::TranscodeList::JobClipboardPrefix);
+	}
+	return false;
+}
+
+void STmvMediaTranscodeList::Paste()
+{
+	UTmvMediaTranscodeList* List = ListHandle ? ListHandle->Get() : nullptr;
+	if (!List)
+	{
+		return;
+	}
+
+	FString PastedText;
+	FPlatformApplicationMisc::ClipboardPaste(PastedText);
+
+	if (!PastedText.StartsWith(UE::TmvMediaEditor::TranscodeList::JobClipboardPrefix))
+	{
+		UE_LOG(LogTmvMediaEditor, Error, TEXT("Paste Jobs: clipboard data is not the expected format."));
+		return;
+	}
+
+	PastedText.RightChopInline(UE::TmvMediaEditor::TranscodeList::JobClipboardPrefix.Len());
+
+	TArray<FTmvMediaTranscodeListItem> Items;
+	FMemoryReaderView Reader(UE::TmvMedia::SerializationUtils::ValueToConstBytesView(PastedText));
+	
+	if (!UE::TmvMedia::TranscodeSerialization::DeserializeTranscodeListItemsFromJson(Reader, Items))
+	{
+		UE_LOG(LogTmvMediaEditor, Error, TEXT("Paste Jobs: Failed to parse clipboard data."));
+		return;
+	}
+
+	if (Items.IsEmpty())
+	{
+		UE_LOG(LogTmvMediaEditor, Error, TEXT("Paste Jobs: Pasted item list is empty."));
+		return;
+	}
+
+	FScopedTransaction Transaction(LOCTEXT("PasteJobTransaction", "Paste Job"));
+	List->Modify();
+	for (const FTmvMediaTranscodeListItem& Item : Items)
+	{
+		List->InsertItemAt(List->GetNumItems(), &Item);
+	}
+}
+
 void STmvMediaTranscodeList::OnTranscodeListItemEvent(const UTmvMediaTranscodeList* InList, const FTmvMediaTranscodeListItemEventArgs& InArgs)
 {
 	// Make sure this is our current list
@@ -629,7 +876,7 @@ void STmvMediaTranscodeList::OnTranscodeListItemEvent(const UTmvMediaTranscodeLi
 		if (JobListView)
 		{
 			JobListView->RebuildList();
-			if (InArgs.Type == ETmvMediaTranscodeListItemEventType::ItemsRemoved)
+			if (InArgs.Type == ETmvMediaTranscodeListItemEventType::ItemsRemoved || InArgs.Type == ETmvMediaTranscodeListItemEventType::ItemsReordered)
 			{
 				JobListView->ClearSelection();
 			}
diff --git a/Engine/Plugins/Media/TmvMedia/Source/TmvMediaEditor/Private/Widgets/STmvMediaTranscodeList.h b/Engine/Plugins/Media/TmvMedia/Source/TmvMediaEditor/Private/Widgets/STmvMediaTranscodeList.h
--- a/Engine/Plugins/Media/TmvMedia/Source/TmvMediaEditor/Private/Widgets/STmvMediaTranscodeList.h
+++ b/Engine/Plugins/Media/TmvMedia/Source/TmvMediaEditor/Private/Widgets/STmvMediaTranscodeList.h
@@ -31,24 +31,50 @@ public:
 
 	void Construct(const FArguments& InArgs);
 
+	//~ Begin SWidget
+	virtual FReply OnKeyDown(const FGeometry& MyGeometry, const FKeyEvent& InKeyEvent) override;
+	//~ End SWidget
+
 private:
+	/** Bind this widget's commands to the command list. */
+	void BindCommands(const TSharedPtr<FUICommandList>& InCommandList);
+	
 	/** JobListView (SListView) row widget generation handler. */
 	TSharedRef<ITableRow> OnGenerateRow(TSharedPtr<FTmvMediaTranscodeJobListItem> InItem, const TSharedRef<STableViewBase>& InOwnerTable) const;
 
 	/** JobListView (SListView) selection change handler. */
 	void OnSelectionChanged(TSharedPtr<FTmvMediaTranscodeJobListItem> InItem, ESelectInfo::Type InSelectInfo);
 
+	/** Returns true if the job view is valid and current item selection is not empty. */
+	bool IsSelectionValid() const;
+
 	/** AddJob Button handlers */
 	bool CanAddJob() const;
 	FReply OnAddJob();
-
+	
 	/** Duplicate Job Button handlers */
-	bool CanDuplicateJob() const;
-	FReply OnDuplicateJob();
+	bool CanDuplicateSelectedJobs() const { return IsSelectionValid(); }
+	FReply OnDuplicateSelectedJobs();
 
 	/** Remove Job Button handlers */
-	bool CanRemoveJob() const;
-	FReply OnRemoveJob();
+	bool CanRemoveSelectedJobs() const { return IsSelectionValid(); }
+	FReply OnRemoveSelectedJobs();
+
+	/** Rename Jobs Action handlers */
+	bool CanRenameSelectedJob() const;
+	void RenameSelectedJob();
+
+	/** Cut Jobs Action handlers */
+	bool CanCutSelectedJobs() const { return IsSelectionValid(); }
+	void CutSelectedJobs();
+
+	/** Copy Jobs Action handlers */
+	bool CanCopySelectedJobs() const { return IsSelectionValid(); }
+	void CopySelectedJobs();
+
+	/** Paste Jobs Action handlers */
+	bool CanPaste() const;
+	void Paste();
 	
 	/** Transcode List Item Event Handler. */
 	void OnTranscodeListItemEvent(const UTmvMediaTranscodeList* InList, const FTmvMediaTranscodeListItemEventArgs& InArgs);
@@ -62,4 +88,7 @@ private:
 
 	/** Job List View items (widget's backing array), reflect the transcode list elements. */
 	TArray<TSharedPtr<FTmvMediaTranscodeJobListItem>> JobList;
+
+	/** Shared Bound Command list for menus, toolbars and buttons. */
+	TSharedPtr<FUICommandList> CommandList;
 };
diff --git a/Engine/Plugins/Media/TmvMedia/Source/TmvMediaEditor/Private/Widgets/STmvMediaTranscoder.cpp b/Engine/Plugins/Media/TmvMedia/Source/TmvMediaEditor/Private/Widgets/STmvMediaTranscoder.cpp
--- a/Engine/Plugins/Media/TmvMedia/Source/TmvMediaEditor/Private/Widgets/STmvMediaTranscoder.cpp
+++ b/Engine/Plugins/Media/TmvMedia/Source/TmvMediaEditor/Private/Widgets/STmvMediaTranscoder.cpp
@@ -22,6 +22,7 @@
 #include "TmvMediaTranscodeListCommands.h"
 #include "Transcoder/TmvMediaTranscodeList.h"
 #include "Transcoder/TmvMediaTranscodeSerialization.h"
+#include "Utils/TmvMediaPathUtils.h"
 #include "Widgets/Layout/SScrollBox.h"
 #include "Widgets/Layout/SSpacer.h"
 #include "Widgets/Layout/SWidgetSwitcher.h"
@@ -63,71 +64,6 @@ namespace UE::TmvMediaEditor::Transcoder
 		EAppReturnType::Type RetVal = FMessageDialog::Open( EAppMsgType::YesNo, DialogText, DialogTitle );
 		return RetVal == EAppReturnType::Yes;
 	}
-
-	/** Ensures the given path starts with a "./". */
-	FString EnsureStartWithDotSlash(const FString& InPath)
-	{
-		if (!InPath.StartsWith(TEXT("./")))
-		{
-			return FPaths::Combine(TEXT("."), InPath);
-		}
-		return InPath;
-	}
-
-	/**
-	 * Try to make the given absolute path relative to project directory.
-	 */
-	bool TryConvertAbsoluteToProjectRelative(FString& InAbsolutePath)
-	{
-		FString ConvertedPath = InAbsolutePath;
-		const FString FullBasePath = FPaths::ConvertRelativePathToFull(FPaths::ProjectDir());
-		if (FPaths::IsUnderDirectory(ConvertedPath, FullBasePath)
-			&& FPaths::MakePathRelativeTo(ConvertedPath, *FullBasePath)) 
-		{
-			InAbsolutePath = EnsureStartWithDotSlash(ConvertedPath);
-			return true;
-		}
-		return false;
-	}
-
-	/**
-	 * Returns a sanitized path compliant with our path resolution rules.
-	 * 
-	 * @param InPickedPath Picked path to an existing file.
-	 * @return sanitized path 
-	 */
-	FString SanitizePickedPath(const FString& InPickedPath)
-	{
-		if (InPickedPath.IsEmpty())
-		{
-			return InPickedPath;
-		}
-
-		// Remove the quotes and normalize for the platform.
-		FString NormalizedPath = InPickedPath.TrimStartAndEnd().Replace(TEXT("\""), TEXT(""));	
-		FPaths::NormalizeDirectoryName(NormalizedPath);
-
-		// The picker usually returns a path relative to executable. Absolute path is preferable.
-		FString SanitizedPath = FPaths::ConvertRelativePathToFull(NormalizedPath);
-
-		// Check if the path is under the current project
-		TryConvertAbsoluteToProjectRelative(SanitizedPath);
-
-		return SanitizedPath;
-	}
-
-	/**
-	 * Resolves the given path to a full path.
-	 * The path can either be absolute or relative to current project dir. 
-	 */
-	FString ResolveFullPath(const FString& InFilePath)
-	{
-		if (InFilePath.StartsWith(TEXT("./")))
-		{
-			return FPaths::ConvertRelativePathToFull(FPaths::Combine(FPaths::ProjectDir(), InFilePath.RightChop(2)));
-		}
-		return InFilePath;
-	}
 }
 
 STmvMediaTranscoder::~STmvMediaTranscoder() = default;
@@ -136,6 +72,7 @@ BEGIN_SLATE_FUNCTION_BUILD_OPTIMIZATION
 void STmvMediaTranscoder::Construct(const FArguments& InArgs)
 {
 	using namespace UE::TmvMediaEditor::Transcoder;
+	using namespace UE::TmvMedia::PathUtils;
 
 	CommandList = MakeShared<FUICommandList>();
 	BindCommands(CommandList);
@@ -143,14 +80,14 @@ void STmvMediaTranscoder::Construct(const FArguments& InArgs)
 	// Get the last value for this setting.
 	GConfig->GetString(*TmvMediaConfigSection, *LastJobItemConfigKey, LastJobItemPath, GEditorPerProjectIni);
 	GConfig->GetString(*TmvMediaConfigSection, *TranscodeListConfigKey, TranscodeListPath, GEditorPerProjectIni);
-	
+
 	// Create a transient transcode list.
 	TranscodeList.Reset(NewObject<UTmvMediaTranscodeList>(GetTransientPackage(), NAME_None));
 	TranscodeListHandle = MakeShared<FTmvMediaTranscodeListHandle>();
 	TranscodeListHandle->SetTranscodeList(TranscodeList.Get());
-	
+
 	// Try to load the last edited config, if any. This is for convenience when there is no asset workflow.
-	const FString ResolvedTranscodeListPath = ResolveFullPath(TranscodeListPath);
+	const FString ResolvedTranscodeListPath = ConvertSanitizedPathToFull(TranscodeListPath);
 	if (FPaths::FileExists(ResolvedTranscodeListPath))
 	{
 		if (!LoadJobList(ResolvedTranscodeListPath))
@@ -214,7 +151,7 @@ void STmvMediaTranscoder::Construct(const FArguments& InArgs)
 							// todo: add '*' if the edited asset is dirty.
 							return FText::Format(LOCTEXT("JobListPathBox", "{0}"), CurrentPathText);
 						})
-					]		
+					]
 				]
 			]
 		]
@@ -279,6 +216,15 @@ void STmvMediaTranscoder::Construct(const FArguments& InArgs)
 }
 END_SLATE_FUNCTION_BUILD_OPTIMIZATION
 
+FReply STmvMediaTranscoder::OnKeyDown(const FGeometry& MyGeometry, const FKeyEvent& InKeyEvent)
+{
+	if (CommandList && CommandList->ProcessCommandBindings(InKeyEvent))
+	{
+		return FReply::Handled();
+	}
+	return SCompoundWidget::OnKeyDown(MyGeometry, InKeyEvent);
+}
+
 int32 STmvMediaTranscoder::GetDetailsViewWidgetIndex() const
 {
 	if (TranscodeListHandle && TranscodeListHandle->GetNumSelected() != 0)
@@ -291,7 +237,7 @@ int32 STmvMediaTranscoder::GetDetailsViewWidgetIndex() const
 void STmvMediaTranscoder::BindCommands(const TSharedPtr<FUICommandList>& InCommandBindings)
 {
 	const FTmvMediaTranscodeListCommands& ListCommands = FTmvMediaTranscodeListCommands::Get();
-	
+
 	InCommandBindings->MapAction(ListCommands.CreateNewJobList,
 		FExecuteAction::CreateSP(this, &STmvMediaTranscoder::OnCreateNewJobList),
 		FCanExecuteAction::CreateSP(this, &STmvMediaTranscoder::CanCreateNewJobList));
@@ -347,7 +293,7 @@ void STmvMediaTranscoder::OnCreateNewJobList()
 void STmvMediaTranscoder::OnOpenJobList()
 {
 	using namespace UE::TmvMediaEditor::Transcoder;
-	
+
 	// Get the saving location
 	TArray<FString> Filenames;
 
@@ -356,14 +302,14 @@ void STmvMediaTranscoder::OnOpenJobList()
 		DesktopPlatform->OpenFileDialog(
 			FSlateApplication::Get().FindBestParentWindowHandleForDialogs(AsShared()),
 			LOCTEXT("OpenJobListDialogTitle", "Open").ToString(),
-			ResolveFullPath(GetJobListPath()),
+			GetJobListBrowseDirectory(),
 			TEXT(""),
 			JobFileTypeFilter,
 			EFileDialogFlags::None,
 			Filenames
 		);
 	}
-	
+
 	if (Filenames.Num() > 0)
 	{
 		// Set the path and load the job list.
@@ -383,7 +329,8 @@ void STmvMediaTranscoder::OnLoadJobList()
 	static const FText LoadFailure = LOCTEXT("LoadListNotificationFail", "Loading Transcode List Failed");
 
 	using namespace UE::TmvMediaEditor::Transcoder;
-	bool bLoaded = LoadJobList(ResolveFullPath(TranscodeListPath));
+	using namespace UE::TmvMedia::PathUtils;
+	bool bLoaded = LoadJobList(ConvertSanitizedPathToFull(TranscodeListPath));
 	AddNotification(bLoaded ? LoadSuccess : LoadFailure, bLoaded);
 }
 
@@ -406,7 +353,7 @@ void STmvMediaTranscoder::OnSaveJobList()
 			return;
 		}
 
-		const FString ResolvedTranscodeListPath = UE::TmvMediaEditor::Transcoder::ResolveFullPath(TranscodeListPath);
+		const FString ResolvedTranscodeListPath = UE::TmvMedia::PathUtils::ConvertSanitizedPathToFull(TranscodeListPath);
 		IFileManager::Get().MakeDirectory(*FPaths::GetPath(ResolvedTranscodeListPath), /*Tree*/ true);		
 		const TUniquePtr<FArchive> FileWriter(IFileManager::Get().CreateFileWriter(*ResolvedTranscodeListPath));
 		if (FileWriter)
@@ -430,7 +377,7 @@ void STmvMediaTranscoder::OnSaveJobList()
 void STmvMediaTranscoder::OnSaveJobListAs()
 {
 	using namespace UE::TmvMediaEditor::Transcoder;
-	
+
 	// Get the saving location
 	TArray<FString> OutFilenames;
 
@@ -439,7 +386,7 @@ void STmvMediaTranscoder::OnSaveJobListAs()
 		DesktopPlatform->SaveFileDialog(
 			FSlateApplication::Get().FindBestParentWindowHandleForDialogs(AsShared()),
 			LOCTEXT("JobListSaveAsDialogTitle", "Save As").ToString(),
-			ResolveFullPath(GetJobListPath()),
+			GetJobListBrowseDirectory(),
 			TEXT(""),
 			JobFileTypeFilter,
 			EFileDialogFlags::None,
@@ -468,7 +415,7 @@ void STmvMediaTranscoder::OnImportJobItemFrom()
 	using namespace UE::TmvMediaEditor::Transcoder;
 
 	int32 SelectedItem = GetFirstSelectedItem();
-	
+
 	if (SelectedItem == INDEX_NONE)
 	{
 		UE_LOG(LogTmvMediaEditor, Error, TEXT("Must select a job item."));
@@ -484,14 +431,14 @@ void STmvMediaTranscoder::OnImportJobItemFrom()
 		DesktopPlatform->OpenFileDialog(
 			FSlateApplication::Get().FindBestParentWindowHandleForDialogs(AsShared()),
 			LOCTEXT("ImportJobItemDialogTitle", "Import Job Item").ToString(),
-			FPaths::GetPath(ResolveFullPath(!GetLastJobItemPath().IsEmpty() ? GetLastJobItemPath() : GetJobListPath())),
+			GetJobItemBrowseDirectory(),
 			TEXT(""),
 			JobFileTypeFilter,
 			EFileDialogFlags::None,
 			Filenames
 		);
 	}
-	
+
 	if (Filenames.Num() > 0)
 	{
 		if (LoadJobItem(SelectedItem, Filenames[0]))
@@ -522,7 +469,7 @@ void STmvMediaTranscoder::OnExportJobItemAs()
 		DesktopPlatform->SaveFileDialog(
 			FSlateApplication::Get().FindBestParentWindowHandleForDialogs(AsShared()),
 			LOCTEXT("ExportJobItemAsDialogTitle", "Export Job Item As").ToString(),
-			FPaths::GetPath(ResolveFullPath(!GetLastJobItemPath().IsEmpty() ? GetLastJobItemPath() : GetJobListPath())),
+			GetJobItemBrowseDirectory(),
 			TEXT(""),
 			JobFileTypeFilter,
 			EFileDialogFlags::None,
@@ -554,7 +501,7 @@ FString STmvMediaTranscoder::GetJobListPath() const
 
 void STmvMediaTranscoder::SetJobListPath(const FString& InPath)
 {
-	TranscodeListPath = UE::TmvMediaEditor::Transcoder::SanitizePickedPath(InPath);
+	TranscodeListPath = UE::TmvMedia::PathUtils::GetSanitizedPath(InPath);
 
 	// Persist the setting.
 	using namespace UE::TmvMediaEditor::Transcoder;
@@ -565,13 +512,24 @@ void STmvMediaTranscoder::SetJobListPath(const FString& InPath)
 void STmvMediaTranscoder::ResetJobListPath()
 {
 	TranscodeListPath.Reset();
-	
+
 	// Persist the setting.
 	using namespace UE::TmvMediaEditor::Transcoder;
 	GConfig->SetString(*TmvMediaConfigSection, *TranscodeListConfigKey, *TranscodeListPath, GEditorPerProjectIni);
 	GConfig->Flush(/*bRemoveFromCache*/ false, GEditorPerProjectIni);
 }
 
+FString STmvMediaTranscoder::GetJobListBrowseDirectory() const
+{
+	if (!TranscodeListPath.IsEmpty())
+	{
+		using namespace UE::TmvMedia::PathUtils;
+		return FPaths::GetPath(ConvertSanitizedPathToFull(TranscodeListPath));
+	}
+
+	return FPaths::ConvertRelativePathToFull(FPaths::ProjectDir());
+}
+
 bool STmvMediaTranscoder::LoadJobList(const FString& InPath)
 {
 	const TUniquePtr<FArchive> FileReader(IFileManager::Get().CreateFileReader(*InPath));
@@ -603,7 +561,7 @@ int32 STmvMediaTranscoder::GetFirstSelectedItem() const
 {
 	if (TranscodeListHandle)
 	{
-		auto Selection = TranscodeListHandle->GetCurrentSelection();
+		TArray<int32> Selection = TranscodeListHandle->GetCurrentSelection();
 		return !Selection.IsEmpty() ? Selection[0] : INDEX_NONE;
 	}
 	return INDEX_NONE;
@@ -616,7 +574,7 @@ FString STmvMediaTranscoder::GetLastJobItemPath() const
 
 void STmvMediaTranscoder::SetLastJobItemPath(const FString& InPath)
 {
-	LastJobItemPath = UE::TmvMediaEditor::Transcoder::SanitizePickedPath(InPath);
+	LastJobItemPath = UE::TmvMedia::PathUtils::GetSanitizedPath(InPath);
 
 	// Persist the setting.
 	using namespace UE::TmvMediaEditor::Transcoder;
@@ -624,6 +582,18 @@ void STmvMediaTranscoder::SetLastJobItemPath(const FString& InPath)
 	GConfig->Flush(/*bRemoveFromCache*/ false, GEditorPerProjectIni);
 }
 
+FString STmvMediaTranscoder::GetJobItemBrowseDirectory() const
+{
+	if (!LastJobItemPath.IsEmpty())
+	{
+		using namespace UE::TmvMedia::PathUtils;
+		return FPaths::GetPath(ConvertSanitizedPathToFull(LastJobItemPath));	
+	}
+
+	return GetJobListBrowseDirectory();
+}
+
+
 bool STmvMediaTranscoder::LoadJobItem(int32 InItemIndex, const FString& InJobItemPath)
 {
 	static const FText LoadSuccess = LOCTEXT("LoadNotificationOk", "Job Item Loaded Successfully");
@@ -638,7 +608,7 @@ bool STmvMediaTranscoder::LoadJobItem(int32 InItemIndex, const FString& InJobIte
 			using namespace UE::TmvMedia::TranscodeSerialization;
 			FTmvMediaTranscodeJobSettings NewJobSettings;
 			TInstancedStruct<FTmvMediaEncoderOptions> NewEncoderOptions;	// Deserializer will initialize.
-	
+
 			if (DeserializeTranscodeJobSettingsFromJson(*FileReader, &NewJobSettings, NewEncoderOptions))
 			{
 				if (FTmvMediaTranscodeListItem* Item = TranscodeList->GetItemMutable(InItemIndex))
@@ -664,7 +634,7 @@ bool STmvMediaTranscoder::SaveJobItem(int32 InItemIndex, const FString& InJobIte
 	bool bSaved = false;
 	if (TranscodeList && TranscodeList->IsValidItemIndex(InItemIndex))
 	{
-		const FString JobItemFullPath = UE::TmvMediaEditor::Transcoder::ResolveFullPath(InJobItemPath);
+		const FString JobItemFullPath = UE::TmvMedia::PathUtils::ConvertSanitizedPathToFull(InJobItemPath);
 		IFileManager::Get().MakeDirectory(*FPaths::GetPath(JobItemFullPath), /*Tree*/ true);
 		
 		const TUniquePtr<FArchive> FileWriter(IFileManager::Get().CreateFileWriter(*JobItemFullPath));
diff --git a/Engine/Plugins/Media/TmvMedia/Source/TmvMediaEditor/Private/Widgets/STmvMediaTranscoder.h b/Engine/Plugins/Media/TmvMedia/Source/TmvMediaEditor/Private/Widgets/STmvMediaTranscoder.h
--- a/Engine/Plugins/Media/TmvMedia/Source/TmvMediaEditor/Private/Widgets/STmvMediaTranscoder.h
+++ b/Engine/Plugins/Media/TmvMedia/Source/TmvMediaEditor/Private/Widgets/STmvMediaTranscoder.h
@@ -24,6 +24,10 @@ public:
 
 	void Construct(const FArguments& InArgs);
 
+	//~ Begin SWidget
+	virtual FReply OnKeyDown(const FGeometry& MyGeometry, const FKeyEvent& InKeyEvent) override;
+	//~ End SWidget
+
 private:
 	/** Details View Widget switcher handler. Returns which widget index to view. */
 	int32 GetDetailsViewWidgetIndex() const;
@@ -56,6 +60,9 @@ private:
 	/** Reset the current job list file path. This will persist in the transcoder config settings. */
 	void ResetJobListPath();
 
+	/** Returns the default browse directory for the job list file dialog(s). */
+	FString GetJobListBrowseDirectory() const;
+
 	/** Load a job list from the given file path. Only support json for now. */
 	bool LoadJobList(const FString& InPath);
 
@@ -74,6 +81,9 @@ private:
 	/** Set the file path of the last exported job item. This will persist in the transcoder config settings. */
 	void SetLastJobItemPath(const FString& InPath);
 
+	/** Returns the default browse directory for the job item file dialog(s). */
+	FString GetJobItemBrowseDirectory() const;
+
 	/** Load (import) a job item in the current list from the given file path (support json format only for now). */
 	bool LoadJobItem(int32 InItemIndex, const FString& InJobItemPath);
 
diff --git a/Engine/Plugins/Media/TmvMedia/Source/TmvMediaEditor/Private/Widgets/STmvMediaTranscoderJobControls.cpp b/Engine/Plugins/Media/TmvMedia/Source/TmvMediaEditor/Private/Widgets/STmvMediaTranscoderJobControls.cpp
--- a/Engine/Plugins/Media/TmvMedia/Source/TmvMediaEditor/Private/Widgets/STmvMediaTranscoderJobControls.cpp
+++ b/Engine/Plugins/Media/TmvMedia/Source/TmvMediaEditor/Private/Widgets/STmvMediaTranscoderJobControls.cpp
@@ -260,7 +260,7 @@ void STmvMediaTranscoderJobControls::CleanUp()
 	}
 	
 	// Also discard all pending jobs.
-	for (const auto& PendingJob : PendingTranscodeJobs )
+	for (const TStrongObjectPtr<UTmvMediaTranscodeJob>& PendingJob : PendingTranscodeJobs )
 	{
 		PendingJob->Discard(CurrentTime);
 	}
diff --git a/Engine/Plugins/Media/TmvMedia/Source/TmvMediaEditor/Private/Widgets/STmvMediaTranscoderJobDetails.cpp b/Engine/Plugins/Media/TmvMedia/Source/TmvMediaEditor/Private/Widgets/STmvMediaTranscoderJobDetails.cpp
--- a/Engine/Plugins/Media/TmvMedia/Source/TmvMediaEditor/Private/Widgets/STmvMediaTranscoderJobDetails.cpp
+++ b/Engine/Plugins/Media/TmvMedia/Source/TmvMediaEditor/Private/Widgets/STmvMediaTranscoderJobDetails.cpp
@@ -2,6 +2,7 @@
 
 #include "STmvMediaTranscoderJobDetails.h"
 
+#include "Customizations/TmvMediaDirectoryPathCustomization.h"
 #include "Customizations/TmvMediaFilePathCustomization.h"
 #include "Editor.h"
 #include "Encoder/TmvMediaEncoderOptions.h"
@@ -115,12 +116,12 @@ bool STmvMediaTranscoderJobDetails::IsDetailsEnabled() const
 	UTmvMediaTranscodeList* List = ListHandle ? ListHandle->Get() : nullptr;
 	if (List && SelectedJobItemIndex != INDEX_NONE)
 	{
-		const auto& Item = List->GetItem(SelectedJobItemIndex);
+		const FTmvMediaTranscodeListItem& Item = List->GetItem(SelectedJobItemIndex);
 		if (const ITmvMediaTranscodeJobManager* TranscodeJobManager = ITmvMediaTranscodeJobManager::Get())
 		{
 			if (TranscodeJobManager->GetTranscodeJob(Item.Id))
 			{
-				// If a job is already queued, it won't be edited anymore (settings are copied in the queued job).
+				// If a job is already queued, it won't be editable anymore (settings are copied in the queued job).
 				return false;
 			}
 		}
@@ -197,9 +198,11 @@ void STmvMediaTranscoderJobDetails::RefreshJobSettingsDetailsView()
 		JobDetailsView = UE::TmvMediaEditor::Transcode::CreateStructureDetailView(JobOptionsWrap);
 		if (JobDetailsView && JobDetailsView->GetDetailsView())
 		{
-			// Add our filepath customization that supports drag and drop and path resolving.
+			// Add our filepath and directory path customizations that supports drag and drop and path resolving.
 			JobDetailsView->GetDetailsView()->RegisterInstancedCustomPropertyTypeLayout(FName(TEXT("FilePath")),
 				FOnGetPropertyTypeCustomizationInstance::CreateStatic(FTmvMediaFilePathCustomization::MakeInstance));
+			JobDetailsView->GetDetailsView()->RegisterInstancedCustomPropertyTypeLayout(FName(TEXT("DirectoryPath")),
+				FOnGetPropertyTypeCustomizationInstance::CreateStatic(FTmvMediaDirectoryPathCustomization::MakeInstance));
 			
 			JobDetailsView->GetDetailsView()->OnFinishedChangingProperties().AddSP(this, &STmvMediaTranscoderJobDetails::OnFinishedChangingProperties);
 		}
diff --git a/Engine/Plugins/Media/TmvMedia/Source/TmvMediaEditor/TmvMediaEditor.Build.cs b/Engine/Plugins/Media/TmvMedia/Source/TmvMediaEditor/TmvMediaEditor.Build.cs
--- a/Engine/Plugins/Media/TmvMedia/Source/TmvMediaEditor/TmvMediaEditor.Build.cs
+++ b/Engine/Plugins/Media/TmvMedia/Source/TmvMediaEditor/TmvMediaEditor.Build.cs
@@ -8,6 +8,7 @@ namespace UnrealBuildTool.Rules
 		{
 			PrivateDependencyModuleNames.AddRange(
 				new string[] {
+					"ApplicationCore",
 					"Core",
 					"CoreUObject",
 					"DesktopWidgets",
