Fixed compilation issue after CL#51209244

#rnx

#Synced-CL 51204997

diff --git a/Engine/Plugins/Media/TmvMedia/Source/TmvMedia/Private/Utils/TmvMediaSerializationUtils.cpp b/Engine/Plugins/Media/TmvMedia/Source/TmvMedia/Private/Utils/TmvMediaSerializationUtils.cpp
--- a/Engine/Plugins/Media/TmvMedia/Source/TmvMedia/Private/Utils/TmvMediaSerializationUtils.cpp
+++ b/Engine/Plugins/Media/TmvMedia/Source/TmvMedia/Private/Utils/TmvMediaSerializationUtils.cpp
@@ -93,7 +93,7 @@ namespace UE::TmvMedia::SerializationUtils
 						UE_LOG(LogTmvMedia, Warning, TEXT("Json Parsing FInstancedStruct Property: Json Object should have only 1 field (has %d)"), (*InstancedStructObject)->Values.Num());
 					}
 						
-					for (const TPair<FString, TSharedPtr<FJsonValue>>& Value : (*InstancedStructObject)->Values)
+					for (const TPair<UE::FSharedString, TSharedPtr<FJsonValue>>& Value : (*InstancedStructObject)->Values)
 					{
 						if (!Value.Value.IsValid())
 						{
@@ -101,7 +101,7 @@ namespace UE::TmvMedia::SerializationUtils
 							continue;
 						}
 
-						FSoftObjectPath ScriptPath(Value.Key);
+						FSoftObjectPath ScriptPath(*Value.Key);
 						if (const UScriptStruct* Struct = Cast<UScriptStruct>(ScriptPath.ResolveObject()))	// so we got a valid struct
 						{
 							if (Value.Value->Type !=  EJson::Object)
