[TmvMedia] Transcoding job settings support more formatting options for output file sequence.

- Supporting {frame_number} token in the base name. Should allow matching movie render queue filename formatting.
- Adding configurable zero padding for frame numbers. ProcessExr is hardcoded to 5 while MRQ uses 4 by default but configurable.

Also included: fixing a failing case in "TmvMedia.SanitizePath" test because ExpandPathTokensToFull was introducing a double slash.

#Synced-CL 51094493

diff --git a/Engine/Plugins/Media/TmvMedia/Source/TmvMedia/Internal/Utils/TmvMediaPathUtils.h b/Engine/Plugins/Media/TmvMedia/Source/TmvMedia/Internal/Utils/TmvMediaPathUtils.h
--- a/Engine/Plugins/Media/TmvMedia/Source/TmvMedia/Internal/Utils/TmvMediaPathUtils.h
+++ b/Engine/Plugins/Media/TmvMedia/Source/TmvMedia/Internal/Utils/TmvMediaPathUtils.h
@@ -47,6 +47,17 @@ namespace UE::TmvMedia::PathUtils
 	 * @return Absolute path
 	 */
 	UE_API FString ConvertSanitizedPathToFull(const FString& InSanitizedPath);
+
+	/**
+	 * Utility function to make a filename from a given base name, possibly including tokens.
+	 * 
+	 * @param InBaseName Base name to use, may contain token {frame_number} (only token supported for now) 
+	 * @param InFrameIndex Frame index (or frame number) of the frame.
+	 * @param InZeroPadFrameNumbers Number of zeros to pad the frame number with.
+	 * @param InExtension file extension to append.
+	 * @return formatted filename with frame index and extension appended. 
+	 */
+	UE_API FString MakeFrameFilename(const FString& InBaseName, int32 InFrameIndex, int32 InZeroPadFrameNumbers, const FString& InExtension);
 }
 
-#undef UE_API
\ No newline at end of file
+#undef UE_API
diff --git a/Engine/Plugins/Media/TmvMedia/Source/TmvMedia/Private/Tests/TmvMediaPathUtilsTests.cpp b/Engine/Plugins/Media/TmvMedia/Source/TmvMedia/Private/Tests/TmvMediaPathUtilsTests.cpp
--- a/Engine/Plugins/Media/TmvMedia/Source/TmvMedia/Private/Tests/TmvMediaPathUtilsTests.cpp
+++ b/Engine/Plugins/Media/TmvMedia/Source/TmvMedia/Private/Tests/TmvMediaPathUtilsTests.cpp
@@ -395,4 +395,59 @@ bool FTmvMediaSanitizePathTests::RunTest(const FString& Parameters)
 	return !HasAnyErrors();
 }
 
+IMPLEMENT_SIMPLE_AUTOMATION_TEST(FTmvMediaMakeFrameFilenameTests, "System.Plugins.TmvMedia.MakeFrameFilename", EAutomationTestFlags::EditorContext | EAutomationTestFlags::EngineFilter)
+bool FTmvMediaMakeFrameFilenameTests::RunTest(const FString& Parameters)
+{
+	// Testing MakeFrameFilename
+	auto TestMakeFrameFilename = [this](const FString& InBaseName, int32 InFrameNumber, int32 InZeroPad, const FString& InExtension, const FString& InExpectedResult, FString& OutError)
+	{
+		const FString FrameFilename = UE::TmvMedia::PathUtils::MakeFrameFilename(InBaseName, InFrameNumber, InZeroPad, InExtension);
+		if (FrameFilename != InExpectedResult)
+		{
+			OutError = FString::Printf(TEXT("MakeFrameFilename Test Failed: BaseName: \"%s\" FrameNumber: %d  ZeroPad: %d Ext: \"%s\" Result: \"%s\" Expected \"%s\""),
+				*InBaseName, InFrameNumber, InZeroPad, *InExtension, *FrameFilename, *InExpectedResult);
+			return false;
+		}
+		return true;
+	};
+
+	FString Error;
+
+	// Test zero pad
+	AddErrorIfFalse(TestMakeFrameFilename(TEXT("Frame_"), 0, 1, TEXT(".exr"), TEXT("Frame_0.exr"), Error), Error);
+	AddErrorIfFalse(TestMakeFrameFilename(TEXT("Frame_"), 0, 2, TEXT(".exr"), TEXT("Frame_00.exr"), Error), Error);
+	AddErrorIfFalse(TestMakeFrameFilename(TEXT("Frame_"), 0, 3, TEXT(".exr"), TEXT("Frame_000.exr"), Error), Error);
+	AddErrorIfFalse(TestMakeFrameFilename(TEXT("Frame_"), 0, 4, TEXT(".exr"), TEXT("Frame_0000.exr"), Error), Error);
+	AddErrorIfFalse(TestMakeFrameFilename(TEXT("Frame_"), 0, 5, TEXT(".exr"), TEXT("Frame_00000.exr"), Error), Error);
+
+	// Test frame index
+	constexpr int32 ZeroPad = 5;
+	AddErrorIfFalse(TestMakeFrameFilename(TEXT("Frame_"), 0, ZeroPad, TEXT(".exr"), TEXT("Frame_00000.exr"), Error), Error);
+	AddErrorIfFalse(TestMakeFrameFilename(TEXT("Frame_"), 0, ZeroPad, TEXT("exr"), TEXT("Frame_00000.exr"), Error), Error);
+	AddErrorIfFalse(TestMakeFrameFilename(TEXT("Frame_"), 123, ZeroPad, TEXT(".exr"), TEXT("Frame_00123.exr"), Error), Error);
+	AddErrorIfFalse(TestMakeFrameFilename(TEXT("Frame_"), 123, ZeroPad, TEXT(".tmv.exr"), TEXT("Frame_00123.tmv.exr"), Error), Error);
+	AddErrorIfFalse(TestMakeFrameFilename(TEXT("Frame."), 123, ZeroPad, TEXT(".exr"), TEXT("Frame.00123.exr"), Error), Error);
+	AddErrorIfFalse(TestMakeFrameFilename(TEXT("Frame."), 123, ZeroPad, TEXT(".tmv.exr"), TEXT("Frame.00123.tmv.exr"), Error), Error);
+	AddErrorIfFalse(TestMakeFrameFilename(TEXT("Frame_"), 123456, ZeroPad, TEXT(".exr"), TEXT("Frame_123456.exr"), Error), Error);
+
+	// {frame_number} token and dot naming
+	AddErrorIfFalse(TestMakeFrameFilename(TEXT("Frame_{frame_number}"), 123, ZeroPad, TEXT(".exr"), TEXT("Frame_00123.exr"), Error), Error);
+	AddErrorIfFalse(TestMakeFrameFilename(TEXT("Frame_{frame_number}"), 123, ZeroPad, TEXT("exr"), TEXT("Frame_00123.exr"), Error), Error);
+	AddErrorIfFalse(TestMakeFrameFilename(TEXT("Frame_{frame_number}"), 123, ZeroPad, TEXT("apv1"), TEXT("Frame_00123.apv1"), Error), Error);
+	AddErrorIfFalse(TestMakeFrameFilename(TEXT("Frame_{frame_number}"), 123, ZeroPad, TEXT(".apv1"), TEXT("Frame_00123.apv1"), Error), Error);
+	AddErrorIfFalse(TestMakeFrameFilename(TEXT("Frame_{frame_number}"), 123, ZeroPad, TEXT(".apv1.exr"), TEXT("Frame_00123.apv1.exr"), Error), Error);
+	AddErrorIfFalse(TestMakeFrameFilename(TEXT("Frame.{frame_number}"), 123, ZeroPad, TEXT("exr"), TEXT("Frame.00123.exr"), Error), Error);
+	AddErrorIfFalse(TestMakeFrameFilename(TEXT("Frame_{frame_number}."), 123, ZeroPad, TEXT(".exr"), TEXT("Frame_00123.exr"), Error), Error);
+	AddErrorIfFalse(TestMakeFrameFilename(TEXT("Frame_{frame_number}."), 123, ZeroPad, TEXT("exr"), TEXT("Frame_00123.exr"), Error), Error);
+	AddErrorIfFalse(TestMakeFrameFilename(TEXT("Frame_{frame_number}_PostFix"), 123, ZeroPad, TEXT(".exr"), TEXT("Frame_00123_PostFix.exr"), Error), Error);
+	AddErrorIfFalse(TestMakeFrameFilename(TEXT("Frame.{frame_number}.PostFix"), 123, ZeroPad, TEXT(".exr"), TEXT("Frame.00123.PostFix.exr"), Error), Error);
+
+	// Test invalid characters /?:&\\*\"<>|%#@^ in either base name or extension.
+	AddErrorIfFalse(TestMakeFrameFilename(TEXT("?Fra?me_?"), 123, ZeroPad, TEXT(".exr"), TEXT("Frame_00123.exr"), Error), Error);
+	AddErrorIfFalse(TestMakeFrameFilename(TEXT("&Fra*me_>"), 123, ZeroPad, TEXT(".exr"), TEXT("Frame_00123.exr"), Error), Error);
+	AddErrorIfFalse(TestMakeFrameFilename(TEXT("\\Fra/me_\""), 123, ZeroPad, TEXT(".&ex*r"), TEXT("Frame_00123.exr"), Error), Error);
+	
+	return !HasAnyErrors();
+}
+
 #endif
diff --git a/Engine/Plugins/Media/TmvMedia/Source/TmvMedia/Private/Transcoder/TmvMediaTmvFrameEncoder.cpp b/Engine/Plugins/Media/TmvMedia/Source/TmvMedia/Private/Transcoder/TmvMediaTmvFrameEncoder.cpp
--- a/Engine/Plugins/Media/TmvMedia/Source/TmvMedia/Private/Transcoder/TmvMediaTmvFrameEncoder.cpp
+++ b/Engine/Plugins/Media/TmvMedia/Source/TmvMedia/Private/Transcoder/TmvMediaTmvFrameEncoder.cpp
@@ -18,6 +18,7 @@
 #include "Transcoder/TmvMediaTranscodeMuxer.h"
 #include "Utils/TTmvMediaSharedObjectPool.h"
 #include "Utils/TmvMediaFrameUtils.h"
+#include "Utils/TmvMediaPathUtils.h"
 
 #include UE_INLINE_GENERATED_CPP_BY_NAME(TmvMediaTmvFrameEncoder)
 
@@ -316,8 +317,8 @@ void UTmvMediaTmvFrameEncoder::ReceiveMips(UTmvMediaTranscodeJob* InParentJob, T
 			return;
 		}
 
-		FString FileName = FString::Printf(TEXT("%s%05d.%s"), *StreamName, InMips->TimeInfo.FrameIndex, *EncoderOptions.Get().GetFileSequenceExtension());
-		FString FullFilename = FPaths::Combine(OutputPath, FPaths::MakeValidFileName(FileName));
+		const FString FileName = UE::TmvMedia::PathUtils::MakeFrameFilename(*StreamName, InMips->TimeInfo.FrameIndex, InParentJob->Settings.ZeroPadFrameNumbers, *EncoderOptions.Get().GetFileSequenceExtension());
+		FString FullFilename = FPaths::Combine(OutputPath, FileName);
 
 		FTmvMediaMemoryEncoderAccessUnit AccessUnit(*AccessUnitBuffer, InMips->TimeInfo.FrameIndex, FullFilename);
 
diff --git a/Engine/Plugins/Media/TmvMedia/Source/TmvMedia/Private/Transcoder/TmvMediaTranscodeMuxer.cpp b/Engine/Plugins/Media/TmvMedia/Source/TmvMedia/Private/Transcoder/TmvMediaTranscodeMuxer.cpp
--- a/Engine/Plugins/Media/TmvMedia/Source/TmvMedia/Private/Transcoder/TmvMediaTranscodeMuxer.cpp
+++ b/Engine/Plugins/Media/TmvMedia/Source/TmvMedia/Private/Transcoder/TmvMediaTranscodeMuxer.cpp
@@ -9,6 +9,7 @@
 #include "Transcoder/TmvMediaFrameMips.h"
 #include "Transcoder/TmvMediaFrameProducer.h"
 #include "Transcoder/TmvMediaTranscodeJob.h"
+#include "Utils/TmvMediaPathUtils.h"
 
 #include UE_INLINE_GENERATED_CPP_BY_NAME(TmvMediaTranscodeMuxer)
 
@@ -100,8 +101,8 @@ void UTmvMediaTranscodeMuxer::ReceiveAccessUnit(
 	// For a container muxer, we need to keep the access units and re-order them. They can be
 	// received out of order because of multiple worker threads. We would need to know the number of worker threads.
 	const FString OutPath = InParentJob->Settings.GetAbsoluteOutputPath();
-	const FString FileName = FString::Printf(TEXT("%s%05d.%s"), *StreamInfo.StreamName, InTimeInfo.FrameIndex, *StreamInfo.Extension);
-	const FString FullFileName = FPaths::Combine(OutPath, FPaths::MakeValidFileName(FileName));
+	const FString FileName = UE::TmvMedia::PathUtils::MakeFrameFilename(*StreamInfo.StreamName, InTimeInfo.FrameIndex, InParentJob->Settings.ZeroPadFrameNumbers, *StreamInfo.Extension);
+	const FString FullFileName = FPaths::Combine(OutPath, FileName);
 
 	TUniquePtr<IFileHandle> File(FPlatformFileManager::Get().GetPlatformFile().OpenWrite(*FullFileName));
 	if (File)
diff --git a/Engine/Plugins/Media/TmvMedia/Source/TmvMedia/Private/Utils/TmvMediaPathUtils.cpp b/Engine/Plugins/Media/TmvMedia/Source/TmvMedia/Private/Utils/TmvMediaPathUtils.cpp
--- a/Engine/Plugins/Media/TmvMedia/Source/TmvMedia/Private/Utils/TmvMediaPathUtils.cpp
+++ b/Engine/Plugins/Media/TmvMedia/Source/TmvMedia/Private/Utils/TmvMediaPathUtils.cpp
@@ -54,7 +54,8 @@ namespace UE::TmvMedia::PathUtils
 				{
 					const FString BaseDirectory = FPaths::ConvertRelativePathToFull(PathToken.GetPathFunction());
 					FString ExpandedPath = FPaths::Combine(BaseDirectory, InPath.RightChop(PathToken.Token.Len()));
-					FPaths::NormalizeDirectoryName(ExpandedPath); // Removes potential double slashes because of Combine.
+					FPaths::NormalizeDirectoryName(ExpandedPath); 
+					FPaths::RemoveDuplicateSlashes(ExpandedPath); // Removes potential double slashes because of Combine.
 					return ExpandedPath;
 				}
 			}
@@ -222,4 +223,43 @@ namespace UE::TmvMedia::PathUtils
 		// Final fallback: try to make it absolute (in case it was still relative to process base dir).
 		return FPaths::ConvertRelativePathToFull(PathExpanded);
 	}
+
+	FString MakeFrameFilename(const FString& InBaseName, int32 InFrameIndex, int32 InZeroPadFrameNumbers, const FString& InExtension)
+	{
+		static const FString FrameNumberToken = TEXT("{frame_number}");
+		
+		// Specify an upper bound (somewhere around typical max filename length for Os) for protection against DOS or OOM.
+		constexpr int32 ZeroPadFrameNumbersUpperBound = 255;
+		const int32 ZeroPadFrameNumber = FMath::Min(InZeroPadFrameNumbers, ZeroPadFrameNumbersUpperBound);
+		FString FormattedName = InBaseName;
+		FString FrameNumber = FString::Printf(TEXT("%0*d"), ZeroPadFrameNumber, InFrameIndex);
+
+		if (FormattedName.Contains(FrameNumberToken))
+		{
+			FormattedName.ReplaceInline(*FrameNumberToken, *FrameNumber);
+		}
+		else
+		{
+			FormattedName += FrameNumber;
+		}
+
+		// Append the extension
+		if (!InExtension.IsEmpty())
+		{
+			if (InExtension.StartsWith(TEXT(".")))
+			{
+				if (FormattedName.EndsWith(TEXT(".")))
+				{
+					FormattedName.LeftChopInline(1);
+				}
+			}
+			else if (!FormattedName.EndsWith(TEXT(".")))
+			{
+				FormattedName += TEXT(".");	
+			}
+			FormattedName += InExtension;
+		}
+		
+		return FPaths::MakeValidFileName(FormattedName);
+	}
 }
\ No newline at end of file
diff --git a/Engine/Plugins/Media/TmvMedia/Source/TmvMedia/Public/Transcoder/TmvMediaTranscodeJob.h b/Engine/Plugins/Media/TmvMedia/Source/TmvMedia/Public/Transcoder/TmvMediaTranscodeJob.h
--- a/Engine/Plugins/Media/TmvMedia/Source/TmvMedia/Public/Transcoder/TmvMediaTranscodeJob.h
+++ b/Engine/Plugins/Media/TmvMedia/Source/TmvMedia/Public/Transcoder/TmvMediaTranscodeJob.h
@@ -92,10 +92,14 @@ struct FTmvMediaTranscodeJobSettings
 	UPROPERTY(EditAnywhere, Category = "Output Sequence")
 	FDirectoryPath OutputPath;
 
-	/** Specify the output base name. */
+	/** Specify the output base name. Supports {frame_number} token. */
 	UPROPERTY(EditAnywhere, Category = "Output Sequence")
 	FString OutputBaseName;
 
+	/** Specify the zero padding on the frame number in the output sequence files. */
+	UPROPERTY(EditAnywhere, Category = "Output Sequence", AdvancedDisplay, meta=(ClampMin=0, UIMin=0))
+	int32 ZeroPadFrameNumbers = 5;
+
 	/** todo - A media source will be made for the sequence. */
 	//UPROPERTY(EditAnywhere, Category = "Output Sequence")
 	//bool bMakeMediaSource = false;
