[TmvMedia] Transcoder Job Editor: implement detail panel multi-selection support.

- also improved transactions and undo handling.

#Synced-CL 51094493

diff --git a/Engine/Plugins/Media/TmvMedia/Source/TmvMedia/Private/Transcoder/TmvMediaTranscodeList.cpp b/Engine/Plugins/Media/TmvMedia/Source/TmvMedia/Private/Transcoder/TmvMediaTranscodeList.cpp
--- a/Engine/Plugins/Media/TmvMedia/Source/TmvMedia/Private/Transcoder/TmvMediaTranscodeList.cpp
+++ b/Engine/Plugins/Media/TmvMedia/Source/TmvMedia/Private/Transcoder/TmvMediaTranscodeList.cpp
@@ -162,4 +162,14 @@ void UTmvMediaTranscodeList::PostLoad()
 	}
 }
 
+#if WITH_EDITOR
+void UTmvMediaTranscodeList::PostEditUndo()
+{
+	Super::PostEditUndo();
+
+	// Trigger a ui refresh
+	GetOnItemEvent().Broadcast(this, {ETmvMediaTranscodeListItemEventType::ItemsModified, {}});
+}
+#endif
+
 
diff --git a/Engine/Plugins/Media/TmvMedia/Source/TmvMedia/Public/Transcoder/TmvMediaTranscodeList.h b/Engine/Plugins/Media/TmvMedia/Source/TmvMedia/Public/Transcoder/TmvMediaTranscodeList.h
--- a/Engine/Plugins/Media/TmvMedia/Source/TmvMedia/Public/Transcoder/TmvMediaTranscodeList.h
+++ b/Engine/Plugins/Media/TmvMedia/Source/TmvMedia/Public/Transcoder/TmvMediaTranscodeList.h
@@ -128,6 +128,9 @@ public:
 	
 	//~ Begin UObject
 	UE_API virtual void PostLoad() override;
+#if WITH_EDITOR
+	UE_API virtual void PostEditUndo() override;
+#endif
 	//~ End UObject
 
 private:
diff --git a/Engine/Plugins/Media/TmvMedia/Source/TmvMediaEditor/Private/Customizations/TmvMediaDirectoryPathCustomization.cpp b/Engine/Plugins/Media/TmvMedia/Source/TmvMediaEditor/Private/Customizations/TmvMediaDirectoryPathCustomization.cpp
--- a/Engine/Plugins/Media/TmvMedia/Source/TmvMediaEditor/Private/Customizations/TmvMediaDirectoryPathCustomization.cpp
+++ b/Engine/Plugins/Media/TmvMedia/Source/TmvMediaEditor/Private/Customizations/TmvMediaDirectoryPathCustomization.cpp
@@ -9,6 +9,7 @@
 #include "Styling/AppStyle.h"
 #include "Utils/TmvMediaPathUtils.h"
 #include "Widgets/STmvMediaDirectoryPicker.h"
+#include "Widgets/TmvMediaEditorTranscodeUtils.h"
 
 #define LOCTEXT_NAMESPACE "FTmvMediaDirectoryPathCustomization"
 
@@ -50,23 +51,26 @@ void FTmvMediaDirectoryPathCustomization::CustomizeChildren(TSharedRef<IProperty
 				.BrowseButtonToolTip(LOCTEXT("DirectoryButtonToolTipText", "Choose a directory from this computer"))
 				.BrowseDirectory_Lambda([this]() -> FString
 				{
-					FString DirectoryPath = GetDirectoryPath();
+					const FString DirectoryPath = GetDirectoryPath(/*bReturnEmptyOnMultiSelect*/ true);
 					return !DirectoryPath.IsEmpty() ?
 						UE::TmvMedia::PathUtils::ConvertSanitizedPathToFull(DirectoryPath) : FPaths::ConvertRelativePathToFull(FPaths::ProjectDir());
 				})
 				.BrowseTitle(LOCTEXT("BrowseTitle", "Directory picker..."))
-				.DirectoryPath(this, &FTmvMediaDirectoryPathCustomization::GetDirectoryPath)	
+				.DirectoryPath(this, &FTmvMediaDirectoryPathCustomization::GetDirectoryPath, /*bReturnEmptyOnMultiSelect*/ false)
 				.OnDirectoryChanged(this, &FTmvMediaDirectoryPathCustomization::HandleDirectoryPickerDirectoryChanged)
 			]
 		];
 }
 
-FString FTmvMediaDirectoryPathCustomization::GetDirectoryPath() const
+FString FTmvMediaDirectoryPathCustomization::GetDirectoryPath(bool bReturnEmptyOnMultiSelect) const
 {
 	FString FilePath;
 	if (PathStringProperty)
 	{
-		PathStringProperty->GetValue(FilePath);
+		if (PathStringProperty->GetValue(FilePath) == FPropertyAccess::Result::MultipleValues)
+		{
+			return !bReturnEmptyOnMultiSelect ? NSLOCTEXT("PropertyEditor", "MultipleValues", "Multiple Values").ToString() : FString();
+		}
 	}
 	return FilePath;
 }
diff --git a/Engine/Plugins/Media/TmvMedia/Source/TmvMediaEditor/Private/Customizations/TmvMediaDirectoryPathCustomization.h b/Engine/Plugins/Media/TmvMedia/Source/TmvMediaEditor/Private/Customizations/TmvMediaDirectoryPathCustomization.h
--- a/Engine/Plugins/Media/TmvMedia/Source/TmvMediaEditor/Private/Customizations/TmvMediaDirectoryPathCustomization.h
+++ b/Engine/Plugins/Media/TmvMedia/Source/TmvMediaEditor/Private/Customizations/TmvMediaDirectoryPathCustomization.h
@@ -32,7 +32,7 @@ private:
 	/**
 	 * Gets the directory path from the property handle.
 	 */
-	FString GetDirectoryPath() const;
+	FString GetDirectoryPath(bool bReturnEmptyOnMultiSelect) const;
 
 	/**
 	 * Called when a directory is selected.
diff --git a/Engine/Plugins/Media/TmvMedia/Source/TmvMediaEditor/Private/Customizations/TmvMediaFilePathCustomization.cpp b/Engine/Plugins/Media/TmvMedia/Source/TmvMediaEditor/Private/Customizations/TmvMediaFilePathCustomization.cpp
--- a/Engine/Plugins/Media/TmvMedia/Source/TmvMediaEditor/Private/Customizations/TmvMediaFilePathCustomization.cpp
+++ b/Engine/Plugins/Media/TmvMedia/Source/TmvMediaEditor/Private/Customizations/TmvMediaFilePathCustomization.cpp
@@ -9,6 +9,7 @@
 #include "Styling/AppStyle.h"
 #include "Utils/TmvMediaPathUtils.h"
 #include "Widgets/Input/SFilePathPicker.h"
+#include "Widgets/TmvMediaEditorTranscodeUtils.h"
 
 #define LOCTEXT_NAMESPACE "FTmvMediaFilePathCustomization"
 
@@ -53,22 +54,27 @@ void FTmvMediaFilePathCustomization::CustomizeChildren(TSharedRef<IPropertyHandl
 				.BrowseButtonToolTip(LOCTEXT("FileButtonToolTipText", "Choose a file from this computer"))
 				.BrowseDirectory_Lambda([this]()->FString
 				{
-					return UE::TmvMedia::PathUtils::GetDirectoryFullPath(HandleFilePathPickerFilePath());
+					const FString FilePath = HandleFilePathPickerFilePath(/*bReturnEmptyOnMultiSelect*/ true);
+					return !FilePath.IsEmpty() ?
+						UE::TmvMedia::PathUtils::GetDirectoryFullPath(FilePath) : FPaths::ConvertRelativePathToFull(FPaths::ProjectDir());
 				})
 				.BrowseTitle(LOCTEXT("PropertyEditorTitle", "File picker..."))
-				.FilePath(this, &FTmvMediaFilePathCustomization::HandleFilePathPickerFilePath)
+				.FilePath(this, &FTmvMediaFilePathCustomization::HandleFilePathPickerFilePath, /*bReturnEmptyOnMultiSelect*/ false)
 				.FileTypeFilter(FileTypeFilter)
 				.OnPathPicked(this, &FTmvMediaFilePathCustomization::HandleFilePathPickerPathPicked)
 			]
 		];
 }
 
-FString FTmvMediaFilePathCustomization::HandleFilePathPickerFilePath() const
+FString FTmvMediaFilePathCustomization::HandleFilePathPickerFilePath(bool bReturnEmptyOnMultiSelect) const
 {
 	FString FilePath;
 	if (PathStringProperty)
 	{
-		PathStringProperty->GetValue(FilePath);
+		if (PathStringProperty->GetValue(FilePath) == FPropertyAccess::Result::MultipleValues)
+		{
+			return !bReturnEmptyOnMultiSelect ? NSLOCTEXT("PropertyEditor", "MultipleValues", "Multiple Values").ToString() : FString();
+		}
 	}
 	return FilePath;
 }
diff --git a/Engine/Plugins/Media/TmvMedia/Source/TmvMediaEditor/Private/Customizations/TmvMediaFilePathCustomization.h b/Engine/Plugins/Media/TmvMedia/Source/TmvMediaEditor/Private/Customizations/TmvMediaFilePathCustomization.h
--- a/Engine/Plugins/Media/TmvMedia/Source/TmvMediaEditor/Private/Customizations/TmvMediaFilePathCustomization.h
+++ b/Engine/Plugins/Media/TmvMedia/Source/TmvMediaEditor/Private/Customizations/TmvMediaFilePathCustomization.h
@@ -32,7 +32,7 @@ private:
 	/**
 	 * Gets the file path.
 	 */
-	FString HandleFilePathPickerFilePath() const;
+	FString HandleFilePathPickerFilePath(bool bReturnEmptyOnMultiSelect) const;
 
 	/**
 	 * Called when a path is selected.
diff --git a/Engine/Plugins/Media/TmvMedia/Source/TmvMediaEditor/Private/Widgets/STmvMediaTranscoderJobDetails.cpp b/Engine/Plugins/Media/TmvMedia/Source/TmvMediaEditor/Private/Widgets/STmvMediaTranscoderJobDetails.cpp
--- a/Engine/Plugins/Media/TmvMedia/Source/TmvMediaEditor/Private/Widgets/STmvMediaTranscoderJobDetails.cpp
+++ b/Engine/Plugins/Media/TmvMedia/Source/TmvMediaEditor/Private/Widgets/STmvMediaTranscoderJobDetails.cpp
@@ -4,11 +4,14 @@
 
 #include "Customizations/TmvMediaDirectoryPathCustomization.h"
 #include "Customizations/TmvMediaFilePathCustomization.h"
+#include "DetailLayoutBuilder.h"
 #include "Editor.h"
 #include "Encoder/TmvMediaEncoderOptions.h"
+#include "IStructureDataProvider.h"
 #include "IStructureDetailsView.h"
 #include "Modules/ModuleManager.h"
 #include "PropertyEditorModule.h"
+#include "ScopedTransaction.h"
 #include "SlateOptMacros.h"
 #include "TmvMediaEditorLog.h"
 #include "TmvMediaEditorTranscodeUtils.h"
@@ -16,11 +19,45 @@
 #include "Transcoder/ITmvMediaTranscodeJobManager.h"
 #include "Transcoder/TmvMediaTranscodeJob.h"
 #include "Transcoder/TmvMediaTranscodeList.h"
-#include "Widgets/Input/STextComboBox.h"
 #include "Widgets/Layout/SBox.h"
 
 #define LOCTEXT_NAMESPACE "STmvMediaTranscoderJobDetails"
 
+namespace UE::TmvMediaEditor::TranscoderJobDetails
+{
+	/**
+	 * Utility function to gather the StructOnScope wrappers of the given Job Items sub-structure.
+	 * @param InListHandle Transcode list handle
+	 * @param InSelectedItemIndices Job Item indices to gather from.
+	 * @param InMakeStructOnScopeForItemFunc Function to make the StructOnScope wrapper from the job item.
+	 */
+	TArray<TSharedPtr<FStructOnScope>> GatherStructOnScopes(
+		const TSharedPtr<FTmvMediaTranscodeListHandle>& InListHandle,
+		const TConstArrayView<int32> InSelectedItemIndices,
+		TFunctionRef<TSharedPtr<FStructOnScope>(FTmvMediaTranscodeListItem*)> InMakeStructOnScopeForItemFunc)
+	{
+		TArray<TSharedPtr<FStructOnScope>> OutStructOnScopes;
+
+		UTmvMediaTranscodeList* List = InListHandle ? InListHandle->Get() : nullptr;
+		if (!List)
+		{
+			return OutStructOnScopes;
+		}
+
+		for(const int32 JobItemIndex : InSelectedItemIndices)
+		{
+			if (FTmvMediaTranscodeListItem* JobItem = List->GetItemMutable(JobItemIndex))
+			{
+				if (TSharedPtr<FStructOnScope> StructOnScope = InMakeStructOnScopeForItemFunc(JobItem))
+				{
+					OutStructOnScopes.Add(MoveTemp(StructOnScope));	
+				}
+			}
+		}
+		return OutStructOnScopes;
+	}
+}
+
 STmvMediaTranscoderJobDetails::~STmvMediaTranscoderJobDetails()
 {
 	if (ListHandle)
@@ -38,7 +75,7 @@ BEGIN_SLATE_FUNCTION_BUILD_OPTIMIZATION
 void STmvMediaTranscoderJobDetails::Construct(const FArguments& InArgs, const TSharedPtr<FTmvMediaTranscodeListHandle>& InListHandle)
 {
 	ListHandle = InListHandle;
-	SelectedJobItemIndex = INDEX_NONE; // nothing selected on construction.
+	CurrentSelection.Reset(); // nothing selected on construction.
 
 	if (ListHandle)
 	{
@@ -74,19 +111,27 @@ void STmvMediaTranscoderJobDetails::Construct(const FArguments& InArgs, const TS
 				SNew(SHorizontalBox)
 				+SHorizontalBox::Slot()
 				.AutoWidth()
-				.Padding(0.0f, 3.0f, 0.0f, 0.0f)
+				.Padding(5.0f, 3.0f, 0.0f, 0.0f)
 				[
 					SNew(STextBlock)
 					.Text(LOCTEXT("TmvMediaEncoderSelectionLabel", "Encoder"))
 				]
 				+ SHorizontalBox::Slot()
 				.AutoWidth()
-				.Padding(5.0f, 0.0f, 0.0f, 0.0f)
+				.Padding(15.0f, 0.0f, 0.0f, 0.0f)
 				[
-					SAssignNew(EncoderOptionsComboBox, STextComboBox)
-					.OptionsSource(&EncoderOptionsNames)
-					.InitiallySelectedItem(GetCurrentEncoderOptionsItem())
+					SAssignNew(EncoderOptionsComboBox, SComboBox<const UScriptStruct*>)
+					.OptionsSource(&EncoderOptionsStructs)
+					.InitiallySelectedItem(GetEncoderOptionsItem())
+					.OnGenerateWidget(this, &STmvMediaTranscoderJobDetails::MakeEncoderOptionsItemWidget)
 					.OnSelectionChanged(this, &STmvMediaTranscoderJobDetails::OnEncoderOptionsSelectionChanged)
+					.ContentPadding(2.f)
+					.Content()
+					[
+						SNew(STextBlock)
+						.Text(this, &STmvMediaTranscoderJobDetails::GetEncoderOptionsContent)
+						.Font(IDetailLayoutBuilder::GetDetailFont())
+					]
 				]
 			]
 		+ SVerticalBox::Slot()
@@ -112,13 +157,14 @@ END_SLATE_FUNCTION_BUILD_OPTIMIZATION
 
 bool STmvMediaTranscoderJobDetails::IsDetailsEnabled() const
 {
-	// The whole details panel is disabled if the job is either queued or running.
+	// The whole details panel is disabled if any of the selected jobs are either queued or running.
 	UTmvMediaTranscodeList* List = ListHandle ? ListHandle->Get() : nullptr;
-	if (List && SelectedJobItemIndex != INDEX_NONE)
+	const ITmvMediaTranscodeJobManager* TranscodeJobManager = ITmvMediaTranscodeJobManager::Get();
+	if (List && !CurrentSelection.IsEmpty() && TranscodeJobManager)
 	{
-		const FTmvMediaTranscodeListItem& Item = List->GetItem(SelectedJobItemIndex);
-		if (const ITmvMediaTranscodeJobManager* TranscodeJobManager = ITmvMediaTranscodeJobManager::Get())
+		for (const int32 SelectedJobItemIndex : CurrentSelection)
 		{
+			const FTmvMediaTranscodeListItem& Item = List->GetItem(SelectedJobItemIndex);
 			if (TranscodeJobManager->GetTranscodeJob(Item.Id))
 			{
 				// If a job is already queued, it won't be editable anymore (settings are copied in the queued job).
@@ -131,7 +177,7 @@ bool STmvMediaTranscoderJobDetails::IsDetailsEnabled() const
 
 void STmvMediaTranscoderJobDetails::OnJobItemSelectionChanged(const UTmvMediaTranscodeList* InList, TConstArrayView<int32> InSelectedItems)
 {
-	SelectedJobItemIndex = !InSelectedItems.IsEmpty() ? InSelectedItems[0] : INDEX_NONE;
+	CurrentSelection = InSelectedItems;
 	RefreshJobSettingsDetailsView();
 	RefreshEncoderOptionsDetailsView();
 	RefreshEncoderOptionsCombo();
@@ -170,32 +216,50 @@ void STmvMediaTranscoderJobDetails::OnTranscodeListChanged(UTmvMediaTranscodeLis
 void STmvMediaTranscoderJobDetails::OnFinishedChangingProperties(const FPropertyChangedEvent& InChangedEvent)
 {
 	UTmvMediaTranscodeList* List = ListHandle ? ListHandle->Get() : nullptr;
-	if (List && SelectedJobItemIndex != INDEX_NONE)
+	if (List && !CurrentSelection.IsEmpty())
 	{
 		TGuardValue TickGuard(bIgnoreModifiedItemEvents, true);	// Ignore events coming from self.
-		List->GetOnItemEvent().Broadcast(List, {ETmvMediaTranscodeListItemEventType::ItemsModified, {SelectedJobItemIndex}});
+		FTmvMediaTranscodeListItemEventArgs Args;
+		Args.Type = ETmvMediaTranscodeListItemEventType::ItemsModified;
+		Args.ItemIndices = CurrentSelection;
+		List->GetOnItemEvent().Broadcast(List, Args);
 	}
+
+	// End Transaction
+	CurrentTransaction.Reset();
 }
 
-FTmvMediaTranscodeListItem* STmvMediaTranscoderJobDetails::GetSelectedJobItem() const
+void STmvMediaTranscoderJobDetails::NotifyPreChange(FEditPropertyChain* InPropertyAboutToChange)
 {
 	UTmvMediaTranscodeList* List = ListHandle ? ListHandle->Get() : nullptr;
-	if (SelectedJobItemIndex == INDEX_NONE || !List)
+	if (List)
 	{
-		return nullptr;
+		// Begin Transaction, if not already in one
+		if (!CurrentTransaction.IsValid())
+		{
+			CurrentTransaction = MakeShared<FScopedTransaction>(LOCTEXT("EditTranscodeJobs", "Edit Transcode Jobs"));
+		}
+		List->Modify();
 	}
-	return List->GetItemMutable(SelectedJobItemIndex);
+}
+
+void STmvMediaTranscoderJobDetails::NotifyPostChange(const FPropertyChangedEvent& InPropertyChangedEvent, FEditPropertyChain* InPropertyThatChanged)
+{
+	// Note: notifying the list is currently done in OnFinishedChangingProperties. We don't have a per property notify hook for now.
 }
 
 void STmvMediaTranscoderJobDetails::RefreshJobSettingsDetailsView()
 {
-	FTmvMediaTranscodeListItem* JobItem = GetSelectedJobItem();
-	TSharedPtr<FStructOnScope> JobOptionsWrap = JobItem ? 
-		MakeShared<FStructOnScope>(FTmvMediaTranscodeJobSettings::StaticStruct(), (uint8*)&JobItem->Settings) : MakeShared<FStructOnScope>();
+	using namespace UE::TmvMediaEditor::TranscoderJobDetails;
+	const TSharedRef<FStructOnScopeStructureDataProvider> StructProvider = MakeShared<FStructOnScopeStructureDataProvider>();
+	StructProvider->SetStructData(GatherStructOnScopes(ListHandle, CurrentSelection, [](FTmvMediaTranscodeListItem* InJobItem)
+	{
+		return MakeShared<FStructOnScope>(FTmvMediaTranscodeJobSettings::StaticStruct(), reinterpret_cast<uint8*>(&InJobItem->Settings));
+	}));
 	
 	if (!JobDetailsView)
 	{
-		JobDetailsView = UE::TmvMediaEditor::Transcode::CreateStructureDetailView(JobOptionsWrap);
+		JobDetailsView = UE::TmvMediaEditor::Transcode::CreateStructureDetailView(StructProvider, this);
 		if (JobDetailsView && JobDetailsView->GetDetailsView())
 		{
 			// Add our filepath and directory path customizations that supports drag and drop and path resolving.
@@ -203,7 +267,7 @@ void STmvMediaTranscoderJobDetails::RefreshJobSettingsDetailsView()
 				FOnGetPropertyTypeCustomizationInstance::CreateStatic(FTmvMediaFilePathCustomization::MakeInstance));
 			JobDetailsView->GetDetailsView()->RegisterInstancedCustomPropertyTypeLayout(FName(TEXT("DirectoryPath")),
 				FOnGetPropertyTypeCustomizationInstance::CreateStatic(FTmvMediaDirectoryPathCustomization::MakeInstance));
-			
+
 			JobDetailsView->GetDetailsView()->OnFinishedChangingProperties().AddSP(this, &STmvMediaTranscoderJobDetails::OnFinishedChangingProperties);
 		}
 
@@ -214,19 +278,23 @@ void STmvMediaTranscoderJobDetails::RefreshJobSettingsDetailsView()
 	}
 	else
 	{
-		JobDetailsView->SetStructureData(JobOptionsWrap);
+		JobDetailsView->SetStructureProvider(StructProvider);
 	}
 }
 
 void STmvMediaTranscoderJobDetails::RefreshEncoderOptionsDetailsView()
 {
-	FTmvMediaTranscodeListItem* JobItem = GetSelectedJobItem();
-	TSharedPtr<FStructOnScope> EncoderOptionsWrap = JobItem && JobItem->EncoderOptions.IsValid() ? 
-		MakeShared<FStructOnScope>(JobItem->EncoderOptions.GetScriptStruct(), JobItem->EncoderOptions.GetMutableMemory()) : MakeShared<FStructOnScope>();
+	using namespace UE::TmvMediaEditor::TranscoderJobDetails;
+	const TSharedRef<FStructOnScopeStructureDataProvider> StructProvider = MakeShared<FStructOnScopeStructureDataProvider>();
+	StructProvider->SetStructData(GatherStructOnScopes(ListHandle, CurrentSelection, [](FTmvMediaTranscodeListItem* InJobItem)
+	{
+		return (InJobItem && InJobItem->EncoderOptions.IsValid()) ?
+			MakeShared<FStructOnScope>(InJobItem->EncoderOptions.GetScriptStruct(), InJobItem->EncoderOptions.GetMutableMemory()) : TSharedPtr<FStructOnScope>();
+	}));
 	
 	if (!EncoderDetailsView)
 	{
-		EncoderDetailsView = UE::TmvMediaEditor::Transcode::CreateStructureDetailView(EncoderOptionsWrap);
+		EncoderDetailsView = UE::TmvMediaEditor::Transcode::CreateStructureDetailView(StructProvider, this);
 		if (!EncoderDetailsView)
 		{
 			UE_LOG(LogTmvMediaEditor, Error, TEXT("Failed to create encoder options detail view."));
@@ -234,79 +302,114 @@ void STmvMediaTranscoderJobDetails::RefreshEncoderOptionsDetailsView()
 	}
 	else
 	{
-		EncoderDetailsView->SetStructureData(EncoderOptionsWrap);
+		EncoderDetailsView->SetStructureProvider(StructProvider);
 	}
 }
 
 void STmvMediaTranscoderJobDetails::RefreshEncoderOptionsCombo()
 {
-	FTmvMediaTranscodeListItem* JobItem = GetSelectedJobItem();
-	
-	const FString SelectedEncoderName = JobItem && JobItem->EncoderOptions.IsValid() ? JobItem->EncoderOptions.GetScriptStruct()->GetDisplayNameText().ToString() : TEXT("None");
-	
-	const int32 NewSelectedEncoderIndex = EncoderOptionsNames.IndexOfByPredicate([&SelectedEncoderName](const TSharedPtr<FString>& InEncoderName)
+	if (EncoderOptionsComboBox)
 	{
-		return *InEncoderName == SelectedEncoderName;
-	});
-
-	if (NewSelectedEncoderIndex != INDEX_NONE)
-	{
-		EncoderOptionsIndex = NewSelectedEncoderIndex;
-		if (EncoderOptionsComboBox)
+		if (const UScriptStruct* SelectedItem = GetEncoderOptionsItem())
 		{
-			EncoderOptionsComboBox->SetSelectedItem(GetCurrentEncoderOptionsItem());
+			EncoderOptionsComboBox->SetSelectedItem(SelectedItem);
 		}
 	}
-	else
-	{
-		UE_LOG(LogTmvMediaEditor, Error, TEXT("Could not find encoder option struct type: %s"), *SelectedEncoderName);
-	}
 }
 
 void STmvMediaTranscoderJobDetails::SetupEncoderOptionsSelection()
 {
 	EncoderOptionsStructs = UE::TmvMediaEditor::Transcode::GetAllDerivedStruct(FTmvMediaEncoderOptions::StaticStruct());
-	EncoderOptionsNames.Reset(EncoderOptionsStructs.Num() + 1);
+}
+
+const UScriptStruct* STmvMediaTranscoderJobDetails::GetEncoderOptionsItem() const
+{
+	UTmvMediaTranscodeList* List = ListHandle ? ListHandle->Get() : nullptr;
+	if (!List)
+	{
+		return nullptr;
+	}
 
-	for (const UScriptStruct* Struct : EncoderOptionsStructs)
+	const UScriptStruct* SelectedStruct = nullptr;
+	
+	for (int32 SelectedJobIndex : CurrentSelection)
 	{
-		EncoderOptionsNames.Add(MakeShared<FString>(Struct->GetDisplayNameText().ToString()));
+		if (FTmvMediaTranscodeListItem* SelectedJobItem = List->GetItemMutable(SelectedJobIndex))
+		{
+			if (!SelectedStruct)
+			{
+				SelectedStruct = SelectedJobItem->EncoderOptions.GetScriptStruct();
+			}
+			else if (SelectedStruct != SelectedJobItem->EncoderOptions.GetScriptStruct())
+			{
+				return nullptr;	// Multiple values
+			}
+		}
 	}
 
-	// Initialize current selection index to the first element.
-	EncoderOptionsIndex = 0;
+	return SelectedStruct;
+}
+
+FText STmvMediaTranscoderJobDetails::GetEncoderOptionsContent() const
+{
+	const UScriptStruct* SelectedItem = GetEncoderOptionsItem();
+	return SelectedItem ? SelectedItem->GetDisplayNameText() : NSLOCTEXT("PropertyEditor", "MultipleValues", "Multiple Values");
+}
+
+TSharedRef<SWidget> STmvMediaTranscoderJobDetails::MakeEncoderOptionsItemWidget(const UScriptStruct* InItem)
+{
+	return SNew(STextBlock)
+		.Text(InItem ? InItem->GetDisplayNameText() : FText::GetEmpty())
+		.Font(IDetailLayoutBuilder::GetDetailFont());
 }
 
-void STmvMediaTranscoderJobDetails::OnEncoderOptionsSelectionChanged(TSharedPtr<FString> InNewSelection, ESelectInfo::Type InSelectInfo)
+void STmvMediaTranscoderJobDetails::OnEncoderOptionsSelectionChanged(const UScriptStruct* InNewSelection, ESelectInfo::Type InSelectInfo)
 {
-	FTmvMediaTranscodeListItem* SelectedJobItem = GetSelectedJobItem();
-	if (!SelectedJobItem)
+	// Ignore if already transacting.
+	if (!InNewSelection || GIsTransacting)
+	{
+		return;
+	}
+
+	UTmvMediaTranscodeList* List = ListHandle ? ListHandle->Get() : nullptr;
+	if (!List)
 	{
-		UE_LOG(LogTmvMediaEditor, Error, TEXT("No Job item selected. Select a job item in the list first."));
+		UE_LOG(LogTmvMediaEditor, Error, TEXT("Invalid Transcode Job List."));
 		return;
 	}
+
+	FScopedTransaction Transaction(LOCTEXT("SetEncoderOptions", "Set Encoder Options"));
+	List->Modify();
 	
-	int32 NewSelectionIndex;
-	if (EncoderOptionsNames.Find(InNewSelection, NewSelectionIndex))
+	int32 NumModifiedJobs = 0;
+
+	for (int32 SelectedJobIndex : CurrentSelection)
 	{
-		if (EncoderOptionsIndex != NewSelectionIndex)
+		FTmvMediaTranscodeListItem* SelectedJobItem = List->GetItemMutable(SelectedJobIndex);	
+		if (!SelectedJobItem)
 		{
-			EncoderOptionsIndex = NewSelectionIndex;
+			continue;
+		}
 
-			if (EncoderDetailsView)
-			{
-				EncoderDetailsView->SetStructureData(MakeShared<FStructOnScope>());	// Clear to avoid dangling.
-			}
-		
-			const UScriptStruct* EncoderOptionsStruct = GetCurrentEncoderOptionsStruct();
-			if (EncoderOptionsStruct && SelectedJobItem->EncoderOptions.GetScriptStruct() != EncoderOptionsStruct)
+		if (SelectedJobItem->EncoderOptions.GetScriptStruct() != InNewSelection)
+		{
+			if (EncoderDetailsView && NumModifiedJobs == 0)
 			{
-				SelectedJobItem->EncoderOptions.InitializeAsScriptStruct(EncoderOptionsStruct);
+				EncoderDetailsView->SetStructureProvider(MakeShared<FStructOnScopeStructureDataProvider>());	// Clear to avoid dangling.
 			}
-
-			RefreshEncoderOptionsDetailsView();
+			SelectedJobItem->EncoderOptions.InitializeAsScriptStruct(InNewSelection);
+			++NumModifiedJobs;
 		}
 	}
+
+	if (NumModifiedJobs)
+	{
+		RefreshEncoderOptionsDetailsView();
+	}
+	else
+	{
+		Transaction.Cancel();
+	}
 }
 
 #undef LOCTEXT_NAMESPACE
diff --git a/Engine/Plugins/Media/TmvMedia/Source/TmvMediaEditor/Private/Widgets/STmvMediaTranscoderJobDetails.h b/Engine/Plugins/Media/TmvMedia/Source/TmvMediaEditor/Private/Widgets/STmvMediaTranscoderJobDetails.h
--- a/Engine/Plugins/Media/TmvMedia/Source/TmvMediaEditor/Private/Widgets/STmvMediaTranscoderJobDetails.h
+++ b/Engine/Plugins/Media/TmvMedia/Source/TmvMediaEditor/Private/Widgets/STmvMediaTranscoderJobDetails.h
@@ -2,19 +2,20 @@
 
 #pragma once
 
+#include "Misc/NotifyHook.h"
+#include "Widgets/Input/SComboBox.h"
 #include "Widgets/SCompoundWidget.h"
 
+class FScopedTransaction;
 class IStructureDetailsView;
-class STextComboBox;
 class UTmvMediaTranscodeList;
 struct FTmvMediaTranscodeListHandle;
-struct FTmvMediaTranscodeListItem;
 struct FTmvMediaTranscodeListItemEventArgs;
 
 /**
  * Transcoder Job Settings and Encoder Options Detail views for the Tmv transcoder list editor. 
  */
-class STmvMediaTranscoderJobDetails : public SCompoundWidget
+class STmvMediaTranscoderJobDetails : public SCompoundWidget, public FNotifyHook
 {
 public:
 	SLATE_BEGIN_ARGS(STmvMediaTranscoderJobDetails){}
@@ -40,8 +41,10 @@ private:
 	/** Details View event handling. Propagates the detail view event. */
 	void OnFinishedChangingProperties(const FPropertyChangedEvent& InChangedEvent);
 
-	/** Get currently "selected" job item, i.e. item that is reflected in the detail views. */
-	FTmvMediaTranscodeListItem* GetSelectedJobItem() const;
+	//~ Begin FNotifyHook
+	virtual void NotifyPreChange(FEditPropertyChain* InPropertyAboutToChange) override;
+	virtual void NotifyPostChange(const FPropertyChangedEvent& InPropertyChangedEvent, FEditPropertyChain* InPropertyThatChanged) override;
+	//~ End FNotifyHook
 
 	/** Refresh (create or update) the Job Setting Details View. */
 	void RefreshJobSettingsDetailsView();
@@ -55,43 +58,30 @@ private:
 	/** Encoder Options Combo Box - Populate Encoder Options Structs list for Combo Box. */
 	void SetupEncoderOptionsSelection();
 
-	/** Encoder Options Combo Box -Returns the currently selected Encoder Options struct. */ 
-	UScriptStruct* GetCurrentEncoderOptionsStruct()
-	{
-		if (EncoderOptionsStructs.IsEmpty())
-		{
-			SetupEncoderOptionsSelection();
-		}
-		if (EncoderOptionsStructs.IsValidIndex(EncoderOptionsIndex))
-		{
-			return EncoderOptionsStructs[EncoderOptionsIndex];
-		}
-		return nullptr;
-	}
+	/**
+	 * Encoder Options Combo Box - Returns the selected "item" from the backing list.
+	 * @return the currently selected struct or null if multiple different values.
+	 */
+	const UScriptStruct* GetEncoderOptionsItem() const;
+
+	/** Encoder Options Combo Box - Returns the text to display in the combo box content. */
+	FText GetEncoderOptionsContent() const;
+
+	/** Encoder Options Combo Box - Returns the widget (a text box) for the given item. */
+	TSharedRef<SWidget> MakeEncoderOptionsItemWidget(const UScriptStruct* InItem);
 
 	/** Encoder Options Combo Box - Handles updating the job item and details view when selection change. */  
-	void OnEncoderOptionsSelectionChanged(TSharedPtr<FString> InNewSelection, ESelectInfo::Type InSelectInfo);
-
-	/** Encoder Options Combo Box - Returns the selected "item" from the backing list. */
-	TSharedPtr<FString> GetCurrentEncoderOptionsItem()
-	{
-		if (EncoderOptionsStructs.IsEmpty())
-		{
-			SetupEncoderOptionsSelection();
-		}
-		if (EncoderOptionsNames.IsValidIndex(EncoderOptionsIndex))
-		{
-			return EncoderOptionsNames[EncoderOptionsIndex];
-		}
-		return nullptr;
-	}
+	void OnEncoderOptionsSelectionChanged(const UScriptStruct* InNewSelection, ESelectInfo::Type InSelectInfo);
 
 	/** Handle to the currently opened Transcode List. */
 	TSharedPtr<FTmvMediaTranscodeListHandle> ListHandle;
 
-	/** Currently selected Job Item Index (in the job list) that is reflected by this detail view. */
-	int32 SelectedJobItemIndex = INDEX_NONE;
+	/** Currently selected Job Item Indices (in the job list) that is reflected by this detail view. */
+	TArray<int32> CurrentSelection;
 
+	/** Current property editing transaction (from any of the details view). */
+	TSharedPtr<FScopedTransaction> CurrentTransaction;
+	
 	/** Flag indicating the modified items events are coming from the detail panel itself and should be ignored. */
 	bool bIgnoreModifiedItemEvents = false;
 	
@@ -102,14 +92,8 @@ private:
 	TSharedPtr<IStructureDetailsView> EncoderDetailsView;
 
 	/** Reference to Encoder Options ComboBox widget to propagate "refresh" events. */
-	TSharedPtr<STextComboBox> EncoderOptionsComboBox;
+	TSharedPtr<SComboBox<const UScriptStruct*>> EncoderOptionsComboBox;
 
 	/** Cached list of the child ScriptStruct of EncoderOptions struct to the encoder options combo box. */
-	TArray<UScriptStruct*> EncoderOptionsStructs;
-
-	/** Display String items (struct display names) for the encoder options combo box. */
-	TArray<TSharedPtr<FString>> EncoderOptionsNames;
-
-	/** Currently selected encoder option's struct for the given job item. */
-	int32 EncoderOptionsIndex = 0;
+	TArray<const UScriptStruct*> EncoderOptionsStructs;
 };
diff --git a/Engine/Plugins/Media/TmvMedia/Source/TmvMediaEditor/Private/Widgets/TmvMediaEditorTranscodeUtils.cpp b/Engine/Plugins/Media/TmvMedia/Source/TmvMediaEditor/Private/Widgets/TmvMediaEditorTranscodeUtils.cpp
--- a/Engine/Plugins/Media/TmvMedia/Source/TmvMediaEditor/Private/Widgets/TmvMediaEditorTranscodeUtils.cpp
+++ b/Engine/Plugins/Media/TmvMedia/Source/TmvMediaEditor/Private/Widgets/TmvMediaEditorTranscodeUtils.cpp
@@ -5,14 +5,13 @@
 #include "IStructureDetailsView.h"
 #include "Modules/ModuleManager.h"
 #include "PropertyEditorModule.h"
-#include "UObject/StructOnScope.h"
 #include "UObject/UObjectIterator.h"
 
 namespace UE::TmvMediaEditor::Transcode
 {
-	TArray<UScriptStruct*> GetAllDerivedStruct(const UScriptStruct* InBaseStruct)
+	TArray<const UScriptStruct*> GetAllDerivedStruct(const UScriptStruct* InBaseStruct)
 	{
-		TArray<UScriptStruct*> OutStructs;
+		TArray<const UScriptStruct*> OutStructs;
 
 		if (InBaseStruct == nullptr)
 		{
@@ -20,19 +19,18 @@ namespace UE::TmvMediaEditor::Transcode
 		}
 
 		OutStructs.Reserve(2);
-		
-		for (TObjectIterator<UScriptStruct> It; It; ++It)
+
+		for (const UScriptStruct* CurrentStruct : TObjectRange<UScriptStruct>())
 		{
-			UScriptStruct* CurrentStruct = *It;
 			if (CurrentStruct != InBaseStruct && CurrentStruct->IsChildOf(InBaseStruct))
 			{
-				OutStructs.Add(CurrentStruct);	
+				OutStructs.Add(CurrentStruct);
 			}
 		}
 		return OutStructs;
 	}
 
-	TSharedPtr<IStructureDetailsView> CreateStructureDetailView(const TSharedPtr<FStructOnScope>& InStructOnScope)
+	TSharedPtr<IStructureDetailsView> CreateStructureDetailView(const TSharedPtr<IStructureDataProvider>& InStructProvider, FNotifyHook* InNotifyHook)
 	{
 		FStructureDetailsViewArgs StructureViewArgs;
 		StructureViewArgs.bShowObjects = true;
@@ -45,8 +43,9 @@ namespace UE::TmvMediaEditor::Transcode
 		ViewArgs.bHideSelectionTip = false;
 		ViewArgs.bShowObjectLabel = false;
 		ViewArgs.NameAreaSettings = FDetailsViewArgs::HideNameArea;
+		ViewArgs.NotifyHook = InNotifyHook;
 
 		FPropertyEditorModule& PropertyEditor = FModuleManager::GetModuleChecked<FPropertyEditorModule>("PropertyEditor");
-		return PropertyEditor.CreateStructureDetailView(ViewArgs, StructureViewArgs, InStructOnScope);
+		return PropertyEditor.CreateStructureProviderDetailView(ViewArgs, StructureViewArgs, InStructProvider);
 	}
 }
\ No newline at end of file
diff --git a/Engine/Plugins/Media/TmvMedia/Source/TmvMediaEditor/Private/Widgets/TmvMediaEditorTranscodeUtils.h b/Engine/Plugins/Media/TmvMedia/Source/TmvMediaEditor/Private/Widgets/TmvMediaEditorTranscodeUtils.h
--- a/Engine/Plugins/Media/TmvMedia/Source/TmvMediaEditor/Private/Widgets/TmvMediaEditorTranscodeUtils.h
+++ b/Engine/Plugins/Media/TmvMedia/Source/TmvMediaEditor/Private/Widgets/TmvMediaEditorTranscodeUtils.h
@@ -5,15 +5,16 @@
 #include "Containers/Array.h"
 #include "Templates/SharedPointer.h"
 
-class FStructOnScope;
+class FNotifyHook;
+class IStructureDataProvider;
 class IStructureDetailsView;
 class UScriptStruct;
 
 namespace UE::TmvMediaEditor::Transcode
 {
 	/** Enumerate all the derived struct. */
-	TArray<UScriptStruct*> GetAllDerivedStruct(const UScriptStruct* InBaseStruct);
-	
+	TArray<const UScriptStruct*> GetAllDerivedStruct(const UScriptStruct* InBaseStruct);
+
 	/** Helper function to create a struct detail view for transcode job details panel. */
-	TSharedPtr<IStructureDetailsView> CreateStructureDetailView(const TSharedPtr<FStructOnScope>& InStructOnScope);
+	TSharedPtr<IStructureDetailsView> CreateStructureDetailView(const TSharedPtr<IStructureDataProvider>& InStructProvider, FNotifyHook* InNotifyHook);
 }
