[ApvMedia] Adding human readable error string to help trouble shooting issues.

#Synced-CL 51094493

diff --git a/Engine/Plugins/Media/TmvMedia/Source/ApvMedia/Private/Apv/ApvDecoder.cpp b/Engine/Plugins/Media/TmvMedia/Source/ApvMedia/Private/Apv/ApvDecoder.cpp
--- a/Engine/Plugins/Media/TmvMedia/Source/ApvMedia/Private/Apv/ApvDecoder.cpp
+++ b/Engine/Plugins/Media/TmvMedia/Source/ApvMedia/Private/Apv/ApvDecoder.cpp
@@ -6,6 +6,49 @@
 
 namespace UE::ApvMedia
 {
+	const TCHAR* GetApvErrorString(int32 InError)
+	{
+		switch (InError)
+		{
+		case OAPV_ERR:
+			return TEXT("Generic Error");
+		case OAPV_ERR_INVALID_ARGUMENT:
+			return TEXT("Invalid argument");
+		case OAPV_ERR_OUT_OF_MEMORY:
+			return TEXT("Out of memory");
+		case OAPV_ERR_REACHED_MAX:
+			return TEXT("reached max");
+		case OAPV_ERR_UNSUPPORTED:
+			return TEXT("Unsupported operation");
+		case OAPV_ERR_UNEXPECTED:
+			return TEXT("Unexpected error");
+		case OAPV_ERR_UNSUPPORTED_COLORSPACE:
+			return TEXT("Unsupported color space");
+		case OAPV_ERR_MALFORMED_BITSTREAM:
+			return TEXT("Malformed bitstream");
+		case OAPV_ERR_OUT_OF_BS_BUF:
+			return TEXT("Bitstream buffer too small");
+		case OAPV_ERR_NOT_FOUND:
+			return TEXT("Not found");
+		case OAPV_ERR_FAILED_SYSCALL:
+			return TEXT("System call failed");
+		case OAPV_ERR_INVALID_PROFILE:
+			return TEXT("Invalid profile");
+		case OAPV_ERR_INVALID_LEVEL:
+			return TEXT("Invalid level");
+		case OAPV_ERR_INVALID_WIDTH:
+			return TEXT("Invalid width");
+		case OAPV_ERR_INVALID_HEIGHT:
+			return TEXT("Invalid height");
+		case OAPV_ERR_INVALID_QP:
+			return TEXT("Invalid QP");
+		case OAPV_ERR_INVALID_FAMILY:
+			return TEXT("Invalid Family Number");
+		default:
+			return TEXT("Unknown error");
+		}
+	}
+
 	FApvImageBitmap::FApvImageBitmap()
 		: ReferenceCount(0)
 	{
diff --git a/Engine/Plugins/Media/TmvMedia/Source/ApvMedia/Private/Apv/ApvDecoder.h b/Engine/Plugins/Media/TmvMedia/Source/ApvMedia/Private/Apv/ApvDecoder.h
--- a/Engine/Plugins/Media/TmvMedia/Source/ApvMedia/Private/Apv/ApvDecoder.h
+++ b/Engine/Plugins/Media/TmvMedia/Source/ApvMedia/Private/Apv/ApvDecoder.h
@@ -7,6 +7,9 @@
 
 namespace UE::ApvMedia
 {
+	/** Returns the error string corresponding to the error code (OAPV_ERR_...) to use for error reporting. */
+	const TCHAR* GetApvErrorString(int32 InError);
+	
 	/**
 	 * Utility wrapper for the decoding image buffer. 
 	 */
diff --git a/Engine/Plugins/Media/TmvMedia/Source/ApvMedia/Private/ApvMediaTmvEncoder.cpp b/Engine/Plugins/Media/TmvMedia/Source/ApvMedia/Private/ApvMediaTmvEncoder.cpp
--- a/Engine/Plugins/Media/TmvMedia/Source/ApvMedia/Private/ApvMediaTmvEncoder.cpp
+++ b/Engine/Plugins/Media/TmvMedia/Source/ApvMedia/Private/ApvMediaTmvEncoder.cpp
@@ -199,7 +199,7 @@ ETmvMediaEncoderResult FApvMediaTmvEncoder::Encode(const FTmvMediaFrameTimeInfo&
 	int32 EncodeResult = oapve_encode(EncoderContext->eid, &InputFrames, EncoderContext->mid, &BitStream, &EncodeStat, &DummyFrames);
 	if (OAPV_FAILED(EncodeResult))
 	{
-		UE_LOG(LogApvMedia, Error, TEXT("Failed to encode frame: Error %d"), EncodeResult);
+		UE_LOG(LogApvMedia, Error, TEXT("Failed to encode frame: Error %s (%d)"), UE::ApvMedia::GetApvErrorString(EncodeResult), EncodeResult);
 		return ETmvMediaEncoderResult::Fail;
 	}
 
