[TmvMedia] Fix Transcoding of DDS sequence with mips.

The converter needs to make sure the destination render target has the RTV flag because DDS media textures don't set it up front.

#Synced-CL 51094493

diff --git a/Engine/Plugins/Media/TmvMedia/Source/TmvMedia/Private/Transcoder/TmvMediaPlayerFrameProducerResource.cpp b/Engine/Plugins/Media/TmvMedia/Source/TmvMedia/Private/Transcoder/TmvMediaPlayerFrameProducerResource.cpp
--- a/Engine/Plugins/Media/TmvMedia/Source/TmvMedia/Private/Transcoder/TmvMediaPlayerFrameProducerResource.cpp
+++ b/Engine/Plugins/Media/TmvMedia/Source/TmvMedia/Private/Transcoder/TmvMediaPlayerFrameProducerResource.cpp
@@ -397,6 +397,7 @@ void FTmvMediaPlayerFrameProducerResource::OnMediaTextureRendered_RenderThread(c
 			TargetMipInfo[0].Planes[0].NumComponents = (TargetMipInfo[0].NumComponents == 1) ? 1 : 4;	// simplify to only support 1 or 4 components.
 			RenderTargetDesc.Format = UE::TmvMedia::FrameUtils::GetPlanePixelFormat(TargetMipInfo[0].Planes[0], /*bNormalized*/ true);
 			RenderTargetDesc.Flags &= ~ETextureCreateFlags::SRGB;	// Remove the SRGB because we will be encoding manually to the desired TF.
+			RenderTargetDesc.Flags |= ETextureCreateFlags::RenderTargetable;	// Make sure the texture has a render target view.
 
 			UE_LOG(LogTmvMedia, Verbose, TEXT("Transcode Frame Producer: Conversion Format: %d x %d, %d mips, %s."),
 				RenderTargetDesc.Extent.X, RenderTargetDesc.Extent.Y, RenderTargetDesc.NumMips, GPixelFormats[RenderTargetDesc.Format].Name);
