[ApvMedia] Apv Encoder: fixing an off by one in the number of mips when creating the encoder context.

Logic/Math error.
Example for a 515 x 513 resolution, it would create up to mip 8 (2x2) and skip the 1x1. That in turn would cause the encoder to fail later when it is asked to encode up to frame 9 (in the same access unit).
Fix: Rewriting the loop in the same way it is done everywhere else.

#Synced-CL 51094493

diff --git a/Engine/Plugins/Media/TmvMedia/Source/ApvMedia/Private/Apv/ApvEncoder.cpp b/Engine/Plugins/Media/TmvMedia/Source/ApvMedia/Private/Apv/ApvEncoder.cpp
--- a/Engine/Plugins/Media/TmvMedia/Source/ApvMedia/Private/Apv/ApvEncoder.cpp
+++ b/Engine/Plugins/Media/TmvMedia/Source/ApvMedia/Private/Apv/ApvEncoder.cpp
@@ -7,6 +7,7 @@
 #include "ApvMipBufferUtils.h"
 #include "Encoder/ITmvMediaEncoder.h"
 #include "Transcoder/TmvMediaFrameMips.h"
+#include "Utils/TmvMediaUtils.h"
 
 namespace UE::ApvMedia
 {
@@ -80,9 +81,10 @@ namespace UE::ApvMedia
 		oapve_cdesc_t EncoderDescriptor;
 		FMemory::Memzero(&EncoderDescriptor, sizeof(oapve_cdesc_t));
 
-		// Primary Frame Parameters (other frames will have the same config)
+		// Configure Primary Frame Parameters (other frames will have the same config)
 		oapve_param_t& PrimaryFrameParams = EncoderDescriptor.param[0];
 		oapve_param_default(&PrimaryFrameParams);
+		EncoderDescriptor.max_num_frms = 1;	// Keep track of the access unit frames we are using.
 
 		// Report invalid frame dimensions.
 		if (InMipInfo.Width <= 0 || InMipInfo.Height <= 0)
@@ -145,51 +147,44 @@ namespace UE::ApvMedia
 			PrimaryFrameParams.full_range_flag = ToUnderlyingType(GetApvColorMatrixRange(InMipInfo.ColorInfo.YuvMatrixRange));
 		}
 
-		int32 NumMips = 0;
-
 		// If the primary frame is already at the smallest possible size, don't add any extra mip frames.
 		if (InMipInfo.bEnableMips && (PrimaryFrameParams.w > 1 || PrimaryFrameParams.h > 1))
 		{
 			const EApvMediaChromaFormat ChromaFormat = InOptions.GetChromaFormat();
+			
+			const int32 MaxNumMips = FTmvMediaUtils::GetMaxMipCountFromDimensions(PrimaryFrameParams.w, PrimaryFrameParams.h);
 
-			for (int32 MipIndex = 0, Width = FMath::Max(PrimaryFrameParams.w/2, 1), Height = FMath::Max(PrimaryFrameParams.h/2, 1);; MipIndex++)
+			// Note: start at frame index 1 because we have already done frame 0 as PrimaryFrame.
+			for (int32 MipIndex = 1; MipIndex < MaxNumMips; ++MipIndex)
 			{
-				// Stop once we've reached the smallest mip. 
-				if (Width == 1 && Height == 1 && MipIndex > 0)
-				{
-					break;
-				}
+				const int32 MipWidth = FMath::Max(1, PrimaryFrameParams.w >> MipIndex);
+				const int32 MipHeight = FMath::Max(1, PrimaryFrameParams.h >> MipIndex);
 
-				if (GetSubWidthC(ChromaFormat) == 2 && Width & 0x1)
+				if (GetSubWidthC(ChromaFormat) == 2 && MipWidth & 0x1)
 				{
-					UE_LOG(LogApvMedia, Warning, TEXT("Can't generate mip of width %d. Not a multiple of two (YUC 422/420 constraint)"), Width);
+					UE_LOG(LogApvMedia, Warning, TEXT("Can't generate mip of width %d. Not a multiple of two (YUC 422/420 constraint)"), MipWidth);
 					break;
 				}
-				if (GetSubHeightC(ChromaFormat) == 2 && Height & 0x1)
+				if (GetSubHeightC(ChromaFormat) == 2 && MipHeight & 0x1)
 				{
-					UE_LOG(LogApvMedia, Warning, TEXT("Can't generate mip of height %d. Not a multiple of two (YUC 420 constraint)"), Height);
+					UE_LOG(LogApvMedia, Warning, TEXT("Can't generate mip of height %d. Not a multiple of two (YUC 420 constraint)"), MipHeight);
 					break;
 				}
 
-				const int32 FrameIndex = 1 + MipIndex; // Assumes one primary frame.
-
-				if (FrameIndex >= OAPV_MAX_NUM_FRAMES) 
-				{ 
+				if (MipIndex >= OAPV_MAX_NUM_FRAMES)
+				{
 					UE_LOG(LogApvMedia, Warning, TEXT("Reached OpenApv maximum number of frames (%d); stopping mip generation."), OAPV_MAX_NUM_FRAMES); 
 					break; 
 				}
 
+				const int32 FrameIndex = EncoderDescriptor.max_num_frms;
 				EncoderDescriptor.param[FrameIndex] = PrimaryFrameParams;	// Copy primary frame parameters
-				EncoderDescriptor.param[FrameIndex].w = Width;
-				EncoderDescriptor.param[FrameIndex].h = Height;
-				++NumMips;
-
-				Width = FMath::Max(Width/2, 1);
-				Height = FMath::Max(Height/2, 1);
+				EncoderDescriptor.param[FrameIndex].w = MipWidth;
+				EncoderDescriptor.param[FrameIndex].h = MipHeight;
+				++EncoderDescriptor.max_num_frms;
 			}
 		}
 
-		EncoderDescriptor.max_num_frms = 1 + NumMips;
 		EncoderDescriptor.threads = GetValidatedNumThreads(InOptions.NumThreads);
 		EncoderDescriptor.max_bs_buf_size = GetDefault<UApvMediaSettings>()->MaxEncoderBitStreamBufferSize;
 
