[TmvMedia] More error handling and reporting improvements

- Adding a message context helper to report the low level errors to the higher level transcoder stages.
- Adding more error message reporting in the lower level of the encoder for validating the formats.
- More error handling around RequestMipInfo to deal with unsupported formats (such as odd dimensions in YUV422).
- Fixing error handling from file sequence frame producer that would leave the job "stopping" forever.

#Synced-CL 51094493

diff --git a/Engine/Plugins/Media/ImgMedia/Source/OpenExrWrapper/Private/OpenExrWrapperTmvEncoder.cpp b/Engine/Plugins/Media/ImgMedia/Source/OpenExrWrapper/Private/OpenExrWrapperTmvEncoder.cpp
--- a/Engine/Plugins/Media/ImgMedia/Source/OpenExrWrapper/Private/OpenExrWrapperTmvEncoder.cpp
+++ b/Engine/Plugins/Media/ImgMedia/Source/OpenExrWrapper/Private/OpenExrWrapperTmvEncoder.cpp
@@ -5,8 +5,11 @@
 #include "OpenExrWrapper.h"
 #include "OpenExrWrapperLog.h"
 #include "SampleConverter/TmvMediaFrameMipBuffer.h"
+#include "Utils/TmvMediaMessageContext.h"
 #include "Utils/TmvMediaUtils.h"
 
+#define LOCTEXT_NAMESPACE "OpenExrTmvEncoder"
+
 // Encoder utility functions
 namespace UE::OpenExrWrapper::Encoder
 {
@@ -81,12 +84,19 @@ FOpenExrWrapperTmvEncoder::FOpenExrWrapperTmvEncoder(const FOpenExrWrapperTmvEnc
 
 FOpenExrWrapperTmvEncoder::~FOpenExrWrapperTmvEncoder() = default;
 
-void FOpenExrWrapperTmvEncoder::RequestMipInfos(const FTmvMediaFrameTimeInfo& InTimeInfo, const FTmvMediaEncoderMipInfo& InFrameInfo, TArray<FTmvMediaFrameMipInfo>& OutFrameMipInfo)
+ETmvMediaEncoderResult FOpenExrWrapperTmvEncoder::RequestMipInfos(
+	const FTmvMediaFrameTimeInfo& InTimeInfo,
+	const FTmvMediaEncoderMipInfo& InFrameInfo,
+	TArray<FTmvMediaFrameMipInfo>& OutFrameMipInfo,
+	FTmvMediaMessageContext* OutMessageContext)
 {
 	if (InFrameInfo.Width <= 0 || InFrameInfo.Height <= 0)
 	{
-		UE_LOG(LogOpenEXRWrapper, Error, TEXT("RequestMipInfos: Invalid frame dimensions %dx%d."), InFrameInfo.Width, InFrameInfo.Height);
-		return;
+		FText Message = FText::Format(
+		LOCTEXT("InvalidFrameDim", "Invalid frame dimensions ({0} x {1})"),
+			FText::AsNumber(InFrameInfo.Width), FText::AsNumber(InFrameInfo.Height));
+		UE_TMV_MEDIA_MESSAGE_LOG(OutMessageContext, LogOpenEXRWrapper, Error, TEXT("RequestMipInfos"), Message);
+		return ETmvMediaEncoderResult::Fail;
 	} 
 
 	const int32 MaxNumMips = FTmvMediaUtils::GetMaxMipCountFromDimensions(InFrameInfo.Width, InFrameInfo.Height);
@@ -154,9 +164,14 @@ void FOpenExrWrapperTmvEncoder::RequestMipInfos(const FTmvMediaFrameTimeInfo& In
 		}
 		OutFrameMipInfo.Add(MipInfo);
 	}
+	return ETmvMediaEncoderResult::Success;
 }
 
-ETmvMediaEncoderResult FOpenExrWrapperTmvEncoder::Encode(const FTmvMediaFrameTimeInfo& InTimeInfo, ITmvMediaEncoderAccessUnit& InAccessUnit, TArrayView<FTmvMediaEncoderMipRequest> InMipRequests)
+ETmvMediaEncoderResult FOpenExrWrapperTmvEncoder::Encode(
+	const FTmvMediaFrameTimeInfo& InTimeInfo,
+	ITmvMediaEncoderAccessUnit& InAccessUnit,
+	TArrayView<FTmvMediaEncoderMipRequest> InMipRequests,
+	FTmvMediaMessageContext* OutMessageContext)
 {
 	// TODO: GetImageParameters. See what DataWindow is and how it is used.
 
@@ -173,7 +188,20 @@ ETmvMediaEncoderResult FOpenExrWrapperTmvEncoder::Encode(const FTmvMediaFrameTim
 
 	if (Mip0Info.Width <= 0 || Mip0Info.Height <= 0)
 	{
-		UE_LOG(LogOpenEXRWrapper, Error, TEXT("Encode: Invalid mip0 dimensions %dx%d."), Mip0Info.Width, Mip0Info.Height);
+		FText Message = FText::Format(
+			LOCTEXT("InvalidMip0Dim", "Invalid mip0 dimensions ({0} x {1})"),
+			FText::AsNumber(Mip0Info.Width), FText::AsNumber(Mip0Info.Height));
+		UE_TMV_MEDIA_MESSAGE_LOG(OutMessageContext, LogOpenEXRWrapper, Error, TEXT("Encode"), Message);
+		return ETmvMediaEncoderResult::Fail;
+	}
+
+	// For now, we only support RGB and RGBA formats.
+	if (Mip0Info.NumComponents > ChannelNames.Num())
+	{
+		FText Message = FText::Format(
+			LOCTEXT("InvalidNumComponents", "The input image has too many components ({0}), maximum supported number of components: {1}"),
+			FText::AsNumber(Mip0Info.NumComponents), FText::AsNumber(ChannelNames.Num()));
+		UE_TMV_MEDIA_MESSAGE_LOG(OutMessageContext, LogOpenEXRWrapper, Error, TEXT("Encode"), Message);
 		return ETmvMediaEncoderResult::Fail;
 	}
 
@@ -198,7 +226,12 @@ ETmvMediaEncoderResult FOpenExrWrapperTmvEncoder::Encode(const FTmvMediaFrameTim
 		for (int32 ComponentIndex = Mip0Info.NumComponents - 1; ComponentIndex >= 0; --ComponentIndex)
 		{
 			int32 PlaneIndex = Mip0Info.GetPlaneIndexForComponent(ComponentIndex);
-			const FTmvMediaFramePlaneInfo& PlaneInfo = Mip0Info.Planes[PlaneIndex]; 
+			if (!ensure(Mip0Info.Planes.IsValidIndex(PlaneIndex)))
+			{
+				continue;
+			}
+
+			const FTmvMediaFramePlaneInfo& PlaneInfo = Mip0Info.Planes[PlaneIndex];
 			const FIntPoint PlaneStride(PlaneInfo.GetBytesPerComponent(), PlaneInfo.Stride);
 			OutFile.AddFrameBufferChannel(ChannelNames[ComponentIndex], nullptr, PlaneStride);
 		}
@@ -218,16 +251,27 @@ ETmvMediaEncoderResult FOpenExrWrapperTmvEncoder::Encode(const FTmvMediaFrameTim
 				UE::OpenExrWrapper::Encoder::TintInplace(MipBuffer, EncoderOptions.MipLevelTints);
 			}
 
-			for (int32 ComponentIndex = 0;  ComponentIndex < MipInfo.NumComponents; ++ComponentIndex)
+			const int32 NumComponents = FMath::Min(MipInfo.NumComponents, ChannelNames.Num());
+
+			for (int32 ComponentIndex = 0;  ComponentIndex < NumComponents; ++ComponentIndex)
 			{
 				void* PlaneBuffer = MipBuffer->GetPlaneBufferForComponent(ComponentIndex);
 				if (PlaneBuffer == nullptr)
 				{
-					UE_LOG(LogOpenEXRWrapper, Error, TEXT("Encode: Null plane buffer for component %d (mip %d)."), ComponentIndex, MipLevel);
+					FText Message = FText::Format(
+						LOCTEXT("NullPlaneBuffer", "Null plane buffer for component {0} (mip {1})"),
+						FText::AsNumber(ComponentIndex), FText::AsNumber(MipLevel));
+					UE_TMV_MEDIA_MESSAGE_LOG(OutMessageContext, LogOpenEXRWrapper, Error, TEXT("Encode"), Message);
 					return ETmvMediaEncoderResult::Fail;
 				}
 
 				const int32 PlaneIndex = MipInfo.GetPlaneIndexForComponent(ComponentIndex);
+
+				if (!ensure(MipInfo.Planes.IsValidIndex(PlaneIndex)))
+				{
+					continue;
+				}
+
 				const FTmvMediaFramePlaneInfo& PlaneInfo = MipInfo.Planes[PlaneIndex];
 				const FIntPoint PlaneStride(PlaneInfo.GetBytesPerComponent(), PlaneInfo.Stride);
 
@@ -243,3 +287,5 @@ ETmvMediaEncoderResult FOpenExrWrapperTmvEncoder::Encode(const FTmvMediaFrameTim
 
 	return ETmvMediaEncoderResult::Success;
 }
+
+#undef LOCTEXT_NAMESPACE
diff --git a/Engine/Plugins/Media/ImgMedia/Source/OpenExrWrapper/Private/OpenExrWrapperTmvEncoder.h b/Engine/Plugins/Media/ImgMedia/Source/OpenExrWrapper/Private/OpenExrWrapperTmvEncoder.h
--- a/Engine/Plugins/Media/ImgMedia/Source/OpenExrWrapper/Private/OpenExrWrapperTmvEncoder.h
+++ b/Engine/Plugins/Media/ImgMedia/Source/OpenExrWrapper/Private/OpenExrWrapperTmvEncoder.h
@@ -15,8 +15,17 @@ public:
 	virtual ~FOpenExrWrapperTmvEncoder() override;
 
 	//~ Begin ITmvMediaEncoder
-	virtual void RequestMipInfos(const FTmvMediaFrameTimeInfo& InTimeInfo, const FTmvMediaEncoderMipInfo& InFrameInfo, TArray<FTmvMediaFrameMipInfo>& OutFrameMipInfo) override;
-	virtual ETmvMediaEncoderResult Encode(const FTmvMediaFrameTimeInfo& InTimeInfo, ITmvMediaEncoderAccessUnit& InAccessUnit, TArrayView<FTmvMediaEncoderMipRequest> InMipRequests) override;
+	virtual ETmvMediaEncoderResult RequestMipInfos(
+		const FTmvMediaFrameTimeInfo& InTimeInfo,
+		const FTmvMediaEncoderMipInfo& InFrameInfo,
+		TArray<FTmvMediaFrameMipInfo>& OutFrameMipInfo,
+		FTmvMediaMessageContext* OutMessageContext) override;
+
+	virtual ETmvMediaEncoderResult Encode(
+		const FTmvMediaFrameTimeInfo& InTimeInfo,
+		ITmvMediaEncoderAccessUnit& InAccessUnit,
+		TArrayView<FTmvMediaEncoderMipRequest> InMipRequests,
+		FTmvMediaMessageContext* OutMessageContext) override;
 	//~ End ITmvMediaEncoder
 
 private:
diff --git a/Engine/Plugins/Media/TmvMedia/Source/ApvMedia/Private/Apv/ApvEncoder.cpp b/Engine/Plugins/Media/TmvMedia/Source/ApvMedia/Private/Apv/ApvEncoder.cpp
--- a/Engine/Plugins/Media/TmvMedia/Source/ApvMedia/Private/Apv/ApvEncoder.cpp
+++ b/Engine/Plugins/Media/TmvMedia/Source/ApvMedia/Private/Apv/ApvEncoder.cpp
@@ -165,19 +165,21 @@ namespace UE::ApvMedia
 					UE_LOG(LogApvMedia, Warning, TEXT("Can't generate mip of width %d. Not a multiple of two (YUC 422/420 constraint)"), MipWidth);
 					break;
 				}
+
 				if (GetSubHeightC(ChromaFormat) == 2 && MipHeight & 0x1)
 				{
 					UE_LOG(LogApvMedia, Warning, TEXT("Can't generate mip of height %d. Not a multiple of two (YUC 420 constraint)"), MipHeight);
 					break;
 				}
 
-				if (MipIndex >= OAPV_MAX_NUM_FRAMES)
+				const int32 FrameIndex = EncoderDescriptor.max_num_frms;
+
+				if (FrameIndex >= OAPV_MAX_NUM_FRAMES)
 				{
-					UE_LOG(LogApvMedia, Warning, TEXT("Reached OpenApv maximum number of frames (%d); stopping mip generation."), OAPV_MAX_NUM_FRAMES); 
-					break; 
+					UE_LOG(LogApvMedia, Warning, TEXT("Reached OpenApv encoder maximum number of frames (%d); stopping mip generation."), OAPV_MAX_NUM_FRAMES); 
+					break;
 				}
 
-				const int32 FrameIndex = EncoderDescriptor.max_num_frms;
 				EncoderDescriptor.param[FrameIndex] = PrimaryFrameParams;	// Copy primary frame parameters
 				EncoderDescriptor.param[FrameIndex].w = MipWidth;
 				EncoderDescriptor.param[FrameIndex].h = MipHeight;
diff --git a/Engine/Plugins/Media/TmvMedia/Source/ApvMedia/Private/ApvMediaTmvEncoder.cpp b/Engine/Plugins/Media/TmvMedia/Source/ApvMedia/Private/ApvMediaTmvEncoder.cpp
--- a/Engine/Plugins/Media/TmvMedia/Source/ApvMedia/Private/ApvMediaTmvEncoder.cpp
+++ b/Engine/Plugins/Media/TmvMedia/Source/ApvMedia/Private/ApvMediaTmvEncoder.cpp
@@ -7,8 +7,11 @@
 #include "ApvMediaSettings.h"
 #include "ImageCoreUtils.h"
 #include "SampleConverter/TmvMediaFrameMipBuffer.h"
+#include "Utils/TmvMediaMessageContext.h"
 #include "Utils/TmvMediaUtils.h"
 
+#define LOCTEXT_NAMESPACE "ApvMediaTmvEncoder"
+
 namespace UE::ApvMedia::Encoder
 {
 	/** Align the given width to Macroblock dimensions. */
@@ -64,13 +67,21 @@ FApvMediaTmvEncoder::FApvMediaTmvEncoder(const FApvMediaTmvEncoderOptions& InEnc
 
 FApvMediaTmvEncoder::~FApvMediaTmvEncoder() = default;
 
-void FApvMediaTmvEncoder::RequestMipInfos(const FTmvMediaFrameTimeInfo& InTimeInfo, const FTmvMediaEncoderMipInfo& InFrameInfo, TArray<FTmvMediaFrameMipInfo>& OutFrameMipInfo)
+ETmvMediaEncoderResult FApvMediaTmvEncoder::RequestMipInfos(
+	const FTmvMediaFrameTimeInfo& InTimeInfo,
+	const FTmvMediaEncoderMipInfo& InFrameInfo,
+	TArray<FTmvMediaFrameMipInfo>& OutFrameMipInfo,
+	FTmvMediaMessageContext* OutMessageContext)
 {
 	using namespace UE::ApvMedia;
 
 	if (InFrameInfo.Width <= 0 || InFrameInfo.Height <= 0)
 	{
-		return;
+		FText Message = FText::Format(
+			LOCTEXT("PrimaryFrameInvalidDim", "Primary Frame has invalid dimensions ({0} x {1})"),
+			FText::AsNumber(InFrameInfo.Width), FText::AsNumber(InFrameInfo.Height));
+		UE_TMV_MEDIA_MESSAGE_LOG(OutMessageContext, LogApvMedia, Error, TEXT("RequestMipInfos"), Message);
+		return ETmvMediaEncoderResult::Fail;
 	}
 
 	const int32 MaxNumMips = FTmvMediaUtils::GetMaxMipCountFromDimensions(InFrameInfo.Width, InFrameInfo.Height);
@@ -84,6 +95,15 @@ void FApvMediaTmvEncoder::RequestMipInfos(const FTmvMediaFrameTimeInfo& InTimeIn
 
 		if (!Encoder::IsMipSizeValid(MipWidth, MipHeight, EncoderOptions.GetChromaFormat()))
 		{
+			// If the first mip (primary frame) has invalid dimensions, this is a fatal error for this operation.
+			if (MipIndex == 0)
+			{
+				FText Message = FText::Format(
+					LOCTEXT("PrimaryFrameInvalidDimForFormat", "Primary Frame has invalid dimensions ({0} x {1}) for the selected chroma format ({2})"),
+					FText::AsNumber(MipWidth), FText::AsNumber(MipHeight), UE::TmvMedia::Utils::StaticEnumToText(EncoderOptions.GetChromaFormat()));
+				UE_TMV_MEDIA_MESSAGE_LOG(OutMessageContext, LogApvMedia, Error, TEXT("RequestMipInfos"), Message);
+				return ETmvMediaEncoderResult::Fail;
+			}
 			break;
 		}
 
@@ -126,9 +146,15 @@ void FApvMediaTmvEncoder::RequestMipInfos(const FTmvMediaFrameTimeInfo& InTimeIn
 		}
 		OutFrameMipInfo.Add(MipInfo);
 	}
+
+	return ETmvMediaEncoderResult::Success;
 }
 
-ETmvMediaEncoderResult FApvMediaTmvEncoder::Encode(const FTmvMediaFrameTimeInfo& InTimeInfo, ITmvMediaEncoderAccessUnit& InAccessUnit, TArrayView<FTmvMediaEncoderMipRequest> InMipRequests)
+ETmvMediaEncoderResult FApvMediaTmvEncoder::Encode(
+	const FTmvMediaFrameTimeInfo& InTimeInfo,
+	ITmvMediaEncoderAccessUnit& InAccessUnit,
+	TArrayView<FTmvMediaEncoderMipRequest> InMipRequests,
+	FTmvMediaMessageContext* OutMessageContext)
 {
 	if (InMipRequests.IsEmpty()) 
 	{ 
@@ -147,8 +173,10 @@ ETmvMediaEncoderResult FApvMediaTmvEncoder::Encode(const FTmvMediaFrameTimeInfo&
 			EncoderContext = MakeUnique<UE::ApvMedia::FApvEncoderContext>(EncoderOptions, EncoderMipInfo, InTimeInfo);
 		}
 		else 
-		{ 
-			UE_LOG(LogApvMedia, Error, TEXT("Encode: Missing mip buffer for index 0.")); 
+		{
+			FText Message = LOCTEXT("MissingMipBufferIndex0", "Missing mip buffer for index 0");
+			UE_TMV_MEDIA_MESSAGE_LOG(OutMessageContext, LogApvMedia, Error, TEXT("Encode"), Message);
+			return ETmvMediaEncoderResult::Fail;
 		}
 	}
 
@@ -164,15 +192,17 @@ ETmvMediaEncoderResult FApvMediaTmvEncoder::Encode(const FTmvMediaFrameTimeInfo&
 	{
 		FTmvMediaFrameMipBufferHandle& MipBuffer = InMipRequests[FrameIndex].MipBuffer;
 		if (!MipBuffer.IsValid()) 
-		{ 
-			UE_LOG(LogApvMedia, Error, TEXT("Encode: Missing mip buffer for index %d."), FrameIndex); 
+		{
+			FText Message = FText::Format(LOCTEXT("MissingMipBufferIndex", "Missing mip buffer for index {0}"), FText::AsNumber(FrameIndex));
+			UE_TMV_MEDIA_MESSAGE_LOG(OutMessageContext, LogApvMedia, Error, TEXT("Encode"), Message);
 			return ETmvMediaEncoderResult::Fail;
 		}
 
 		InputFrames.frm[FrameIndex].imgb = UE::ApvMedia::ApvSetupImageBuffer(MipBuffer->GetMipInfoRef(), *MipBuffer);
 		if (InputFrames.frm[FrameIndex].imgb == nullptr)
-		{ 
-			UE_LOG(LogApvMedia, Error, TEXT("Encode: Failed to set up OpenAPV image buffer for index %d."), FrameIndex);
+		{
+			FText Message = FText::Format(LOCTEXT("FailedAllocApvImageBuffer", "Failed to set up OpenAPV image buffer for index {0}"), FText::AsNumber(FrameIndex));
+			UE_TMV_MEDIA_MESSAGE_LOG(OutMessageContext, LogApvMedia, Error, TEXT("Encode"), Message);
 			return ETmvMediaEncoderResult::Fail;
 		}
 
@@ -199,14 +229,20 @@ ETmvMediaEncoderResult FApvMediaTmvEncoder::Encode(const FTmvMediaFrameTimeInfo&
 	int32 EncodeResult = oapve_encode(EncoderContext->eid, &InputFrames, EncoderContext->mid, &BitStream, &EncodeStat, &DummyFrames);
 	if (OAPV_FAILED(EncodeResult))
 	{
-		UE_LOG(LogApvMedia, Error, TEXT("Failed to encode frame: Error %s (%d)"), UE::ApvMedia::GetApvErrorString(EncodeResult), EncodeResult);
+		FText Message = FText::Format(
+			LOCTEXT("FailedEncodeFrame", "Failed to encode frame: Error {0} ({1})"),
+			FText::FromString(UE::ApvMedia::GetApvErrorString(EncodeResult)), FText::AsNumber(EncodeResult));
+		UE_TMV_MEDIA_MESSAGE_LOG(OutMessageContext, LogApvMedia, Error, TEXT("Encode"), Message);
 		return ETmvMediaEncoderResult::Fail;
 	}
 
 	// Validate the written data size.
 	if (EncodeStat.write < 0 || EncodeStat.write > BitStreamBuffer.Num())
 	{
-		UE_LOG(LogApvMedia, Error, TEXT("Encoded data size (%d) overflows encoded buffer size (%lld)"), EncodeStat.write, BitStreamBuffer.Num());
+		FText Message = FText::Format(
+			LOCTEXT("EncodedDataOverflow", "Encoded data size ({0}) overflows encoded buffer size ({1})"),
+			FText::AsNumber(EncodeStat.write), FText::AsNumber(BitStreamBuffer.Num()));
+		UE_TMV_MEDIA_MESSAGE_LOG(OutMessageContext, LogApvMedia, Error, TEXT("Encode"), Message);
 		return ETmvMediaEncoderResult::Fail;
 	}
 
@@ -214,3 +250,5 @@ ETmvMediaEncoderResult FApvMediaTmvEncoder::Encode(const FTmvMediaFrameTimeInfo&
 	InAccessUnit.Write(BitStreamBuffer.GetData(), EncodeStat.write);
 	return ETmvMediaEncoderResult::Success;
 }
+
+#undef LOCTEXT_NAMESPACE
\ No newline at end of file
diff --git a/Engine/Plugins/Media/TmvMedia/Source/ApvMedia/Private/ApvMediaTmvEncoder.h b/Engine/Plugins/Media/TmvMedia/Source/ApvMedia/Private/ApvMediaTmvEncoder.h
--- a/Engine/Plugins/Media/TmvMedia/Source/ApvMedia/Private/ApvMediaTmvEncoder.h
+++ b/Engine/Plugins/Media/TmvMedia/Source/ApvMedia/Private/ApvMediaTmvEncoder.h
@@ -20,8 +20,17 @@ public:
 	virtual ~FApvMediaTmvEncoder() override;
 
 	//~ Begin ITmvMediaEncoder
-	virtual void RequestMipInfos(const FTmvMediaFrameTimeInfo& InTimeInfo, const FTmvMediaEncoderMipInfo& InFrameInfo, TArray<FTmvMediaFrameMipInfo>& OutFrameMipInfo) override;
-	virtual ETmvMediaEncoderResult Encode(const FTmvMediaFrameTimeInfo& InTimeInfo, ITmvMediaEncoderAccessUnit& InAccessUnit, TArrayView<FTmvMediaEncoderMipRequest> InMipRequests) override;
+	virtual ETmvMediaEncoderResult RequestMipInfos(
+		const FTmvMediaFrameTimeInfo& InTimeInfo,
+		const FTmvMediaEncoderMipInfo& InFrameInfo,
+		TArray<FTmvMediaFrameMipInfo>& OutFrameMipInfo,
+		FTmvMediaMessageContext* OutMessageContext) override;
+
+	virtual ETmvMediaEncoderResult Encode(
+		const FTmvMediaFrameTimeInfo& InTimeInfo,
+		ITmvMediaEncoderAccessUnit& InAccessUnit,
+		TArrayView<FTmvMediaEncoderMipRequest> InMipRequests,
+		FTmvMediaMessageContext* OutMessageContext) override;
 	//~ End ITmvMediaEncoder
 
 private:
diff --git a/Engine/Plugins/Media/TmvMedia/Source/TmvMedia/Internal/Transcoder/TmvMediaTmvFrameEncoder.h b/Engine/Plugins/Media/TmvMedia/Source/TmvMedia/Internal/Transcoder/TmvMediaTmvFrameEncoder.h
--- a/Engine/Plugins/Media/TmvMedia/Source/TmvMedia/Internal/Transcoder/TmvMediaTmvFrameEncoder.h
+++ b/Engine/Plugins/Media/TmvMedia/Source/TmvMedia/Internal/Transcoder/TmvMediaTmvFrameEncoder.h
@@ -24,7 +24,7 @@ public:
 
 	//~ Begin UTmvMediaFrameEncoder
 	virtual void SetEncoderOptions(const TInstancedStruct<FTmvMediaEncoderOptions>& InOptions) override;
-	virtual void RequestMipInfos(UTmvMediaTranscodeJob* InParentJob, const FTmvMediaFrameTimeInfo& InSourceTimeInfo, const FTmvMediaFrameMipInfo& InSourceMipInfo, TArray<FTmvMediaFrameMipInfo>& OutFrameMipInfo) override;
+	virtual bool RequestMipInfos(UTmvMediaTranscodeJob* InParentJob, const FTmvMediaFrameTimeInfo& InSourceTimeInfo, const FTmvMediaFrameMipInfo& InSourceMipInfo, TArray<FTmvMediaFrameMipInfo>& OutFrameMipInfo) override;
 	virtual void ReceiveMips(UTmvMediaTranscodeJob* InParentJob, TUniquePtr<FTmvMediaFrameMips>&& InMips) override;
 	//~ End UTmvMediaFrameEncoder
 
diff --git a/Engine/Plugins/Media/TmvMedia/Source/TmvMedia/Private/Transcoder/TmvMediaFileSequenceFrameProducer.cpp b/Engine/Plugins/Media/TmvMedia/Source/TmvMedia/Private/Transcoder/TmvMediaFileSequenceFrameProducer.cpp
--- a/Engine/Plugins/Media/TmvMedia/Source/TmvMedia/Private/Transcoder/TmvMediaFileSequenceFrameProducer.cpp
+++ b/Engine/Plugins/Media/TmvMedia/Source/TmvMedia/Private/Transcoder/TmvMediaFileSequenceFrameProducer.cpp
@@ -100,8 +100,14 @@ bool UTmvMediaFileSequenceFrameProducer::Start(UTmvMediaTranscodeJob* InParentJo
 
 void UTmvMediaFileSequenceFrameProducer::RequestStop(UTmvMediaTranscodeJob* InParentJob)
 {
-	// Stopping is not immediate, need to wait for the processing thread to return.
-	SetStageStatus(ETmvMediaTranscodeStageStatus::Stopping, InParentJob);
+	if (NumActiveFrames > 0)
+	{
+		// Stopping is not immediate, need to wait for the processing thread to return.
+		SetStageStatus(ETmvMediaTranscodeStageStatus::Stopping, InParentJob);
+		return;
+	}
+
+	Super::RequestStop(InParentJob);
 }
 
 void UTmvMediaFileSequenceFrameProducer::ProcessAllImages(UTmvMediaTranscodeJob* InParentJob, const FString& InSequencePath)
diff --git a/Engine/Plugins/Media/TmvMedia/Source/TmvMedia/Private/Transcoder/TmvMediaFrameConverter.cpp b/Engine/Plugins/Media/TmvMedia/Source/TmvMedia/Private/Transcoder/TmvMediaFrameConverter.cpp
--- a/Engine/Plugins/Media/TmvMedia/Source/TmvMedia/Private/Transcoder/TmvMediaFrameConverter.cpp
+++ b/Engine/Plugins/Media/TmvMedia/Source/TmvMedia/Private/Transcoder/TmvMediaFrameConverter.cpp
@@ -680,7 +680,10 @@ void UTmvMediaFrameConverter::ReceiveMips(UTmvMediaTranscodeJob* InParentJob, TU
 		}
 
 		TRACE_CPUPROFILER_EVENT_SCOPE(FrameConverter::RequestMipInfos);
-		FrameEncoder->RequestMipInfos(InParentJob, InMips->TimeInfo, MipInfo, EncoderFrameMipInfo);
+		if (!FrameEncoder->RequestMipInfos(InParentJob, InMips->TimeInfo, MipInfo, EncoderFrameMipInfo))
+		{
+			return;
+		}
 	}
 
 	TArray<FTmvMediaFrameMipBufferHandle> NewMipBuffers;
diff --git a/Engine/Plugins/Media/TmvMedia/Source/TmvMedia/Private/Transcoder/TmvMediaPlayerFrameProducerResource.cpp b/Engine/Plugins/Media/TmvMedia/Source/TmvMedia/Private/Transcoder/TmvMediaPlayerFrameProducerResource.cpp
--- a/Engine/Plugins/Media/TmvMedia/Source/TmvMedia/Private/Transcoder/TmvMediaPlayerFrameProducerResource.cpp
+++ b/Engine/Plugins/Media/TmvMedia/Source/TmvMedia/Private/Transcoder/TmvMediaPlayerFrameProducerResource.cpp
@@ -380,7 +380,13 @@ void FTmvMediaPlayerFrameProducerResource::OnMediaTextureRendered_RenderThread(c
 			{
 				FTmvMediaFrameMipInfo MediaTextureMipInfo;
 				UE::TmvMedia::PopulateMipInfoFromTexture(MediaTextureRhi, 0, MediaTextureMipInfo);
-				FrameEncoder->RequestMipInfos(ParentJob, TimeInfo, MediaTextureMipInfo, TargetMipInfo);
+
+				TargetMipInfo.Reset();
+
+				if (!FrameEncoder->RequestMipInfos(ParentJob, TimeInfo, MediaTextureMipInfo, TargetMipInfo))
+				{
+					return;
+				}
 			}
 		}
 
@@ -389,7 +395,7 @@ void FTmvMediaPlayerFrameProducerResource::OnMediaTextureRendered_RenderThread(c
 		UE_LOG(LogTmvMedia, Verbose, TEXT("Transcode Frame Producer: Media Texture Format: %d x %d, %d mips, %s."),
 			RenderTargetDesc.Extent.X, RenderTargetDesc.Extent.Y, RenderTargetDesc.NumMips, GPixelFormats[RenderTargetDesc.Format].Name);
 		
-		if (TargetMipInfo.Num())
+		if (TargetMipInfo.IsValidIndex(0))
 		{
 			// Note: The destination render target number of mips matches the media texture, not TargetMipInfo.
 
@@ -421,7 +427,7 @@ void FTmvMediaPlayerFrameProducerResource::OnMediaTextureRendered_RenderThread(c
 	{
 		RHI_BREADCRUMB_EVENT_STAT_F(RhiCmdList, TmvMediaFrameProducer_ConvertFrame, "TmvMediaFrameProducer.ConvertFrame", "TmvMediaFrameProducer.ConvertFrame %d", TimeInfo.FrameIndex);
 		FRDGTextureRef RenderTargetRdg = GraphBuilder.RegisterExternalTexture(CreateRenderTarget(RenderTargetRhi, TEXT("RenderTarget")));
-		if (TargetMipInfo.Num())
+		if (TargetMipInfo.IsValidIndex(0))
 		{
 			UE::TmvMedia::AddConvertTexturePass(GraphBuilder, MediaTextureRdg, RenderTargetRdg, TargetMipInfo[0]);
 			bGpuConverted = true;
@@ -442,7 +448,7 @@ void FTmvMediaPlayerFrameProducerResource::OnMediaTextureRendered_RenderThread(c
 			SubmitReadbackRequest(GraphBuilder, ReadbackRdg, TimeInfo, ReadbackRequest);
 
 			// If the texture was converted by the GPU, it should have converted the color space, color model and encoded.
-			if (bGpuConverted)
+			if (bGpuConverted && TargetMipInfo.IsValidIndex(0))
 			{
 				ReadbackRequest->MipInfo.ColorInfo = TargetMipInfo[0].ColorInfo;
 				ReadbackRequest->MipInfo.ColorModel = TargetMipInfo[0].ColorModel;
diff --git a/Engine/Plugins/Media/TmvMedia/Source/TmvMedia/Private/Transcoder/TmvMediaTmvFrameEncoder.cpp b/Engine/Plugins/Media/TmvMedia/Source/TmvMedia/Private/Transcoder/TmvMediaTmvFrameEncoder.cpp
--- a/Engine/Plugins/Media/TmvMedia/Source/TmvMedia/Private/Transcoder/TmvMediaTmvFrameEncoder.cpp
+++ b/Engine/Plugins/Media/TmvMedia/Source/TmvMedia/Private/Transcoder/TmvMediaTmvFrameEncoder.cpp
@@ -16,6 +16,7 @@
 #include "Transcoder/TmvMediaFrameMips.h"
 #include "Transcoder/TmvMediaTranscodeJob.h"
 #include "Transcoder/TmvMediaTranscodeMuxer.h"
+#include "Utils/TmvMediaMessageContext.h"
 #include "Utils/TTmvMediaSharedObjectPool.h"
 #include "Utils/TmvMediaFrameUtils.h"
 #include "Utils/TmvMediaPathUtils.h"
@@ -213,12 +214,12 @@ void UTmvMediaTmvFrameEncoder::SetEncoderOptions(const TInstancedStruct<FTmvMedi
 	EncoderOptions = InOptions;
 }
 
-void UTmvMediaTmvFrameEncoder::RequestMipInfos(UTmvMediaTranscodeJob* InParentJob, const FTmvMediaFrameTimeInfo& InSourceTimeInfo, const FTmvMediaFrameMipInfo& InSourceMipInfo, TArray<FTmvMediaFrameMipInfo>& OutFrameMipInfo)
+bool UTmvMediaTmvFrameEncoder::RequestMipInfos(UTmvMediaTranscodeJob* InParentJob, const FTmvMediaFrameTimeInfo& InSourceTimeInfo, const FTmvMediaFrameMipInfo& InSourceMipInfo, TArray<FTmvMediaFrameMipInfo>& OutFrameMipInfo)
 {
 	if (!InParentJob)
 	{
 		UE_LOG(LogTmvMedia, Error, TEXT("Frame Encoder Stage: RequestMipInfos called with invalid Parent Job."));
-		return;
+		return false;
 	}
 
 	FTmvMediaEncoderMipInfo EncoderMipInfo;
@@ -249,17 +250,26 @@ void UTmvMediaTmvFrameEncoder::RequestMipInfos(UTmvMediaTranscodeJob* InParentJo
 	if (!EncoderPoolRef.IsValid())
 	{
 		UE_LOG(LogTmvMedia, Log, TEXT("Frame Encoder Stage: Encoder Pool not available."));
-		return;
+		return false;
 	}
 
 	FTmvMediaTmvEncoderPool::FHandle EncoderHandle = EncoderPoolRef->AcquireEncoder(TEXT(""), EncoderOptions);
 	if (!EncoderHandle.IsValid())
 	{
 		UE_LOG(LogTmvMedia, Error, TEXT("Frame Encoder Stage: Failed to acquire encoder \"%s\"."), *EncoderOptions.Get().GetEncoderName().ToString());
-		return;
+		return false;
 	}
 
-	EncoderHandle->RequestMipInfos(InSourceTimeInfo, EncoderMipInfo, OutFrameMipInfo);
+	FTmvMediaMessageContext MessageContext;
+	const ETmvMediaEncoderResult Result = EncoderHandle->RequestMipInfos(InSourceTimeInfo, EncoderMipInfo, OutFrameMipInfo, &MessageContext);
+	if (Result != ETmvMediaEncoderResult::Success)
+	{
+		const FText Message = MessageContext.ToText();
+		UE_LOG(LogTmvMedia, Error, TEXT("Frame Encoder Stage Error: %s."), *Message.ToString());
+		UTmvMediaTranscodeJob::SafeReportError(InParentJob, Message);
+		return false;
+	}
+	return true;
 }
 
 void UTmvMediaTmvFrameEncoder::ReceiveMips(UTmvMediaTranscodeJob* InParentJob, TUniquePtr<FTmvMediaFrameMips>&& InMips)
@@ -337,10 +347,13 @@ void UTmvMediaTmvFrameEncoder::ReceiveMips(UTmvMediaTranscodeJob* InParentJob, T
 		}
 
 		TRACE_CPUPROFILER_EVENT_SCOPE(UTmvMediaTmvFrameEncoder::Encode);
-		ETmvMediaEncoderResult Result = EncoderHandle->Encode(InMips->TimeInfo, AccessUnit, MipRequests);
+		FTmvMediaMessageContext MessageContext;
+		const ETmvMediaEncoderResult Result = EncoderHandle->Encode(InMips->TimeInfo, AccessUnit, MipRequests, &MessageContext);
 		if (Result != ETmvMediaEncoderResult::Success)
 		{
-			const FText Message = FText::Format(LOCTEXT("FailedEncodeFrame", "Failed to encode frame {0}"), FText::AsNumber(InMips->TimeInfo.FrameIndex));
+			const FText Message = FText::Format(
+				LOCTEXT("FailedEncodeFrame", "Failed to encode frame {0}: {1}"),
+				FText::AsNumber(InMips->TimeInfo.FrameIndex), MessageContext.ToText());
 			UE_LOG(LogTmvMedia, Error, TEXT("Frame Encoder Stage: %s."), *Message.ToString());
 			UTmvMediaTranscodeJob::SafeReportError(InParentJob, Message);
 			return;
diff --git a/Engine/Plugins/Media/TmvMedia/Source/TmvMedia/Private/Transcoder/TmvMediaTranscodeJob.cpp b/Engine/Plugins/Media/TmvMedia/Source/TmvMedia/Private/Transcoder/TmvMediaTranscodeJob.cpp
--- a/Engine/Plugins/Media/TmvMedia/Source/TmvMedia/Private/Transcoder/TmvMediaTranscodeJob.cpp
+++ b/Engine/Plugins/Media/TmvMedia/Source/TmvMedia/Private/Transcoder/TmvMediaTranscodeJob.cpp
@@ -164,8 +164,6 @@ void UTmvMediaTranscodeJob::RequestStop(double InCurrentTime, ETmvMediaTranscode
 {
 	UE_LOG(LogTmvMedia, Verbose, TEXT("Transcode Job for \"%s\" Stopped. Reason: %s"),
 		*Settings.GetInputPath(), *UE::TmvMedia::Utils::StaticEnumToString(InReason)); 
-
-	const bool bWasRunning = IsRunning();
 	
 	if (IsRunning())
 	{
diff --git a/Engine/Plugins/Media/TmvMedia/Source/TmvMedia/Private/Utils/TmvMediaMessageContext.cpp b/Engine/Plugins/Media/TmvMedia/Source/TmvMedia/Private/Utils/TmvMediaMessageContext.cpp
new file mode 100644
--- /dev/null
+++ b/Engine/Plugins/Media/TmvMedia/Source/TmvMedia/Private/Utils/TmvMediaMessageContext.cpp
@@ -0,0 +1,14 @@
+// Copyright Epic Games, Inc. All Rights Reserved.
+
+#include "Utils/TmvMediaMessageContext.h"
+
+
+FText FTmvMediaMessageContext::ToText() const
+{
+	if (Messages.IsEmpty())
+	{
+		return FText::GetEmpty();
+	}
+
+	return FText::Join(FText::FromString("\n"), Messages);
+}
diff --git a/Engine/Plugins/Media/TmvMedia/Source/TmvMedia/Public/Encoder/ITmvMediaEncoder.h b/Engine/Plugins/Media/TmvMedia/Source/TmvMedia/Public/Encoder/ITmvMediaEncoder.h
--- a/Engine/Plugins/Media/TmvMedia/Source/TmvMedia/Public/Encoder/ITmvMediaEncoder.h
+++ b/Engine/Plugins/Media/TmvMedia/Source/TmvMedia/Public/Encoder/ITmvMediaEncoder.h
@@ -10,6 +10,7 @@
 
 class FArchive;
 class FString;
+class FTmvMediaMessageContext;
 struct FTmvMediaFrameTimeInfo;
 
 /**
@@ -125,24 +126,32 @@ public:
 
 	/**
 	 * Queries the encoder for the memory layout and format for frames of the given size, and doing so for each mips.
+	 * Can fail if the encoder doesn't support the given input format.
+	 * 
 	 * @param InTimeInfo Time information needed by some encoders to estimate bitrate and compression.
 	 * @param InFrameInfo Mip0 information of incoming frame.
 	 * @param OutFrameMipInfo output array populated with the request mip info.
+	 * @param OutMessageContext Optional message context for the given operation.
+	 * @return Result code depending on the success of the operation.
 	 */
-	virtual void RequestMipInfos(
+	virtual ETmvMediaEncoderResult RequestMipInfos(
 		const FTmvMediaFrameTimeInfo& InTimeInfo,
 		const FTmvMediaEncoderMipInfo& InFrameInfo,
-		TArray<FTmvMediaFrameMipInfo>& OutFrameMipInfo) = 0;
+		TArray<FTmvMediaFrameMipInfo>& OutFrameMipInfo,
+		FTmvMediaMessageContext* OutMessageContext) = 0;
 
 	/**
 	 * Encode the given mip requests.
+	 * 
 	 * @param InTimeInfo Time information needed by some encoders to encode in the frame.
 	 * @param InAccessUnit Access unit for the current frame's encoded stream.
 	 * @param InMipRequests Array of mip encoding requests.
+	 * @param OutMessageContext Optional message context for the given operation.
 	 * @return Result code depending on the success of the operation.
 	 */
 	virtual ETmvMediaEncoderResult Encode(
 		const FTmvMediaFrameTimeInfo& InTimeInfo,
 		ITmvMediaEncoderAccessUnit& InAccessUnit,
-		TArrayView<FTmvMediaEncoderMipRequest> InMipRequests) = 0;
+		TArrayView<FTmvMediaEncoderMipRequest> InMipRequests,
+		FTmvMediaMessageContext* OutMessageContext) = 0;
 };
\ No newline at end of file
diff --git a/Engine/Plugins/Media/TmvMedia/Source/TmvMedia/Public/Transcoder/TmvMediaFrameEncoder.h b/Engine/Plugins/Media/TmvMedia/Source/TmvMedia/Public/Transcoder/TmvMediaFrameEncoder.h
--- a/Engine/Plugins/Media/TmvMedia/Source/TmvMedia/Public/Transcoder/TmvMediaFrameEncoder.h
+++ b/Engine/Plugins/Media/TmvMedia/Source/TmvMedia/Public/Transcoder/TmvMediaFrameEncoder.h
@@ -35,8 +35,12 @@ public:
 	 * @param InSourceTimeInfo Source sample time info, needed for frame duration (frame rate). Encoder might need it to compute some bitrate related settings.
 	 * @param InSourceMipInfo Source Mip Info from the prior stage (ex Frame Producer).
 	 * @param OutFrameMipInfo Populated mip info with the memory layout the encoder needs.
+	 * @return true if the requested info can be produced, or false if there is a validation error or the encoder doesn't support requested format.
 	 */
-	virtual void RequestMipInfos(UTmvMediaTranscodeJob* InParentJob, const FTmvMediaFrameTimeInfo& InSourceTimeInfo, const FTmvMediaFrameMipInfo& InSourceMipInfo, TArray<FTmvMediaFrameMipInfo>& OutFrameMipInfo) {}
+	virtual bool RequestMipInfos(UTmvMediaTranscodeJob* InParentJob, const FTmvMediaFrameTimeInfo& InSourceTimeInfo, const FTmvMediaFrameMipInfo& InSourceMipInfo, TArray<FTmvMediaFrameMipInfo>& OutFrameMipInfo) 
+	{
+		return false;
+	}
 
 	/**
 	 * Receive mips for processing.
diff --git a/Engine/Plugins/Media/TmvMedia/Source/TmvMedia/Public/Utils/TmvMediaMessageContext.h b/Engine/Plugins/Media/TmvMedia/Source/TmvMedia/Public/Utils/TmvMediaMessageContext.h
new file mode 100644
--- /dev/null
+++ b/Engine/Plugins/Media/TmvMedia/Source/TmvMedia/Public/Utils/TmvMediaMessageContext.h
@@ -0,0 +1,61 @@
+// Copyright Epic Games, Inc. All Rights Reserved.
+
+#pragma once
+
+#include "Containers/Array.h"
+#include "Internationalization/Text.h"
+#include "Logging/LogMacros.h"
+
+#define UE_API TMVMEDIA_API
+
+/**
+ * Container for error messages that get accumulated during operations and can be compiled
+ * into a report for display in UI or dumped to Logs.
+ */
+class FTmvMediaMessageContext
+{
+public:
+	/** Adds the given message at the end of the report. */
+	void Add(const FText& InMessage)
+	{
+		Messages.Add(InMessage);
+	}
+
+	/** Adds the given message at the end of the report. */
+	void Add(FText&& InMessage)
+	{
+		Messages.Add(MoveTemp(InMessage));
+	}
+
+	/**
+	 * Compiles the set of errors in a formatted report with each error message on a different line.
+	 */
+	UE_API FText ToText() const;
+	
+	/**
+	 * List of error messages for the given context.
+	 */
+	TArray<FText> Messages;
+};
+
+/**
+ * Helper macro to reduce boilerplate for message reporting.
+ * In most cases, we want to both log the message in a log category and add it to the report.
+ * We have to use a macro for this to preserve the file/line functionality of the logging system.
+ * 
+ * @param MessageContext Message Report Context to accumulate messages
+ * @param LogCategory Log category to log to.
+ * @param Verbosity Verbosity of the log.
+ * @param SectionString Name of the section this log will be part of.
+ * @param MessageText Formatted message.
+ */
+#define UE_TMV_MEDIA_MESSAGE_LOG(MessageContext, LogCategory, Verbosity, SectionString, MessageText) \
+	{ \
+		UE_LOG(LogCategory, Verbosity, TEXT("%s: %s."), SectionString, *MessageText.ToString()); \
+		if (MessageContext) \
+		{ \
+			MessageContext->Add(MoveTemp(MessageText)); \
+		} \
+	}
+
+#undef UE_API
diff --git a/Engine/Plugins/Media/TmvMedia/Source/TmvMedia/Public/Utils/TmvMediaUtils.h b/Engine/Plugins/Media/TmvMedia/Source/TmvMedia/Public/Utils/TmvMediaUtils.h
--- a/Engine/Plugins/Media/TmvMedia/Source/TmvMedia/Public/Utils/TmvMediaUtils.h
+++ b/Engine/Plugins/Media/TmvMedia/Source/TmvMedia/Public/Utils/TmvMediaUtils.h
@@ -52,4 +52,14 @@ namespace UE::TmvMedia::Utils
 		}
 		return TEXT("<invalid-enum>");
 	}
+
+	template<typename InEnumType>
+	FText StaticEnumToText(InEnumType InValue)
+	{
+		if (UEnum* Enum = StaticEnum<InEnumType>())
+		{
+			return Enum->GetDisplayNameTextByValue(static_cast<int64>(InValue));
+		}
+		return NSLOCTEXT("TmvMediaUtils", "InvalidEnumText", "<invalid-enum>");
+	}
 }
\ No newline at end of file
diff --git a/Engine/Plugins/Media/TmvMedia/Source/TmvMediaEditor/Private/Widgets/TmvMediaTranscodeNotification.h b/Engine/Plugins/Media/TmvMedia/Source/TmvMediaEditor/Private/Widgets/TmvMediaTranscodeNotification.h
--- a/Engine/Plugins/Media/TmvMedia/Source/TmvMediaEditor/Private/Widgets/TmvMediaTranscodeNotification.h
+++ b/Engine/Plugins/Media/TmvMedia/Source/TmvMediaEditor/Private/Widgets/TmvMediaTranscodeNotification.h
@@ -52,6 +52,11 @@ private:
 		{
 			NotificationItem->SetEnabled(false);
 			NotificationItem->SetCompletionState(bInSuccess ? SNotificationItem::CS_Success : SNotificationItem::CS_Fail);
+			if (!bInSuccess)
+			{
+				// Leave more time to read error messages.
+				NotificationItem->SetExpireDuration(2.5f);
+			}
 			NotificationItem->ExpireAndFadeout();
 			NotificationItem.Reset();
 		}
