[TmvMedia] Fix for conversion of rectangular sequences APV in format YUV444/4 (i.e. full mip chain)

#Synced-CL 51094493

diff --git a/Engine/Plugins/Media/TmvMedia/Source/ApvMedia/Private/ApvMediaTmvEncoder.cpp b/Engine/Plugins/Media/TmvMedia/Source/ApvMedia/Private/ApvMediaTmvEncoder.cpp
--- a/Engine/Plugins/Media/TmvMedia/Source/ApvMedia/Private/ApvMediaTmvEncoder.cpp
+++ b/Engine/Plugins/Media/TmvMedia/Source/ApvMedia/Private/ApvMediaTmvEncoder.cpp
@@ -90,8 +90,8 @@ ETmvMediaEncoderResult FApvMediaTmvEncoder::RequestMipInfos(
 
 	for (int32 MipIndex = 0; MipIndex < NumMips; ++MipIndex)
 	{
-		const int32 MipWidth = InFrameInfo.Width >> MipIndex;
-		const int32 MipHeight = InFrameInfo.Height >> MipIndex;
+		const int32 MipWidth = FMath::Max(1, InFrameInfo.Width >> MipIndex);
+		const int32 MipHeight = FMath::Max(1, InFrameInfo.Height >> MipIndex);
 
 		if (!Encoder::IsMipSizeValid(MipWidth, MipHeight, EncoderOptions.GetChromaFormat()))
 		{
