[TmvMediaEditor] Transcoder Job List will select the newly created jobs.

#Synced-CL 51094493

diff --git a/Engine/Plugins/Media/TmvMedia/Source/TmvMediaEditor/Private/Widgets/STmvMediaTranscodeList.cpp b/Engine/Plugins/Media/TmvMedia/Source/TmvMediaEditor/Private/Widgets/STmvMediaTranscodeList.cpp
--- a/Engine/Plugins/Media/TmvMedia/Source/TmvMediaEditor/Private/Widgets/STmvMediaTranscodeList.cpp
+++ b/Engine/Plugins/Media/TmvMedia/Source/TmvMediaEditor/Private/Widgets/STmvMediaTranscodeList.cpp
@@ -464,6 +464,122 @@ namespace UE::TmvMediaEditor::TranscodeList
 			.Text(InCommand->GetLabel())
 			.ToolTipText(InCommand->GetDescription());
 	}
+
+	/**
+	 * Utility class that will accumulate added item ids during an operation (from
+	 * multiple Item Events) and select those items when the operation is completed.
+	 * 
+	 * This method works with operations that add items in multiple calls to the transcode list
+	 * which generate an Item Event per item added, but we want to select all added items and not
+	 * just the last one.
+	 */
+	class FAddedItemsScopedSelectionModifier
+	{
+	public:
+		FAddedItemsScopedSelectionModifier(
+			const TSharedPtr<FTmvMediaTranscodeListHandle>& InListHandle,
+			const TSharedPtr<SListView<TSharedPtr<FTmvMediaTranscodeJobListItem>>>& InJobListView)
+			: ListHandle(InListHandle)
+			, JobListView(InJobListView)
+		{
+			if (UTmvMediaTranscodeList* List = ListHandle ? ListHandle->Get() : nullptr)
+			{
+				List->GetOnItemEvent().AddRaw(this, &FAddedItemsScopedSelectionModifier::OnTranscodeListItemEvent);
+			}
+		}
+
+		~FAddedItemsScopedSelectionModifier()
+		{
+			if (UTmvMediaTranscodeList* List = ListHandle ? ListHandle->Get() : nullptr)
+			{
+				List->GetOnItemEvent().RemoveAll(this);	
+			}
+
+			if (!JobListView)
+			{
+				return;
+			}
+
+			JobListView->ClearSelection();
+
+			if (AddedItemIds.IsEmpty())
+			{
+				return;
+			}
+			
+			TArray<TSharedPtr<FTmvMediaTranscodeJobListItem>> NewSelectedItems;
+			NewSelectedItems.Reserve(AddedItemIds.Num());
+
+			const TArrayView<const TSharedPtr<FTmvMediaTranscodeJobListItem>>& Items = JobListView->GetItems();
+
+			if (AddedItemIds.Num() > 1)
+			{
+				// Use a map to speed up the search in case of a larger list.
+				TMap<FGuid, const TSharedPtr<FTmvMediaTranscodeJobListItem>> ItemMapByIds;
+
+				ItemMapByIds.Reserve(Items.Num());
+				for (const TSharedPtr<FTmvMediaTranscodeJobListItem>& Item : Items)
+				{
+					ItemMapByIds.Add(Item->ItemId, Item);
+				}
+				
+				for (const FGuid& ItemId : AddedItemIds)
+				{
+					if (const TSharedPtr<FTmvMediaTranscodeJobListItem>* FoundItem = ItemMapByIds.Find(ItemId))
+					{
+						NewSelectedItems.Add(*FoundItem);
+					}
+				}
+			}
+			else
+			{
+				for (const FGuid& ItemId : AddedItemIds)
+				{
+					const TSharedPtr<FTmvMediaTranscodeJobListItem>* FoundItem =
+						Items.FindByPredicate([ItemId](const TSharedPtr<FTmvMediaTranscodeJobListItem>& InItem)
+						{
+							return InItem->ItemId == ItemId;
+						});
+
+					if (FoundItem)
+					{
+						NewSelectedItems.Add(*FoundItem);
+					}
+				}
+			}
+
+			if (NewSelectedItems.Num())
+			{
+				JobListView->SetItemSelection(NewSelectedItems, /*bSelected*/ true);
+			}
+		}
+		
+	private:
+		void OnTranscodeListItemEvent(const UTmvMediaTranscodeList* InList, const FTmvMediaTranscodeListItemEventArgs& InArgs)
+		{
+			if (InArgs.Type == ETmvMediaTranscodeListItemEventType::ItemsAdded && InList)
+			{
+				// Indices can change in case of multiple insertions, for safety, we keep track of ids instead.
+				for (int32 ItemIndex : InArgs.ItemIndices)
+				{
+					const FTmvMediaTranscodeListItem& Item = InList->GetItem(ItemIndex);
+					if (Item.Id.IsValid())
+					{
+						AddedItemIds.AddUnique(Item.Id);
+					}
+				}
+			}
+		}
+
+		/** Accumulator of added item Ids.*/
+		TArray<FGuid> AddedItemIds;
+
+		/** Reference to the list handle. */
+		TSharedPtr<FTmvMediaTranscodeListHandle> ListHandle;
+
+		/** JobListView to modify selection. */
+		TSharedPtr<SListView<TSharedPtr<FTmvMediaTranscodeJobListItem>>> JobListView;
+	};
 }
 
 BEGIN_SLATE_FUNCTION_BUILD_OPTIMIZATION
@@ -701,6 +817,8 @@ FReply STmvMediaTranscodeList::OnAddJob()
 	{
 		FScopedTransaction Transaction(LOCTEXT("AddJobTransaction", "Add Job"));
 		List->Modify();
+
+		UE::TmvMediaEditor::TranscodeList::FAddedItemsScopedSelectionModifier SelectionModifier(ListHandle, JobListView);
 		List->InsertItemAt(List->GetNumItems(), nullptr);
 	}
 	return FReply::Handled();
@@ -712,7 +830,8 @@ FReply STmvMediaTranscodeList::OnDuplicateSelectedJobs()
 	{
 		FScopedTransaction Transaction(LOCTEXT("DuplicateSelectionTransaction", "Duplicate Selected Job(s)"));
 		List->Modify();
-		
+
+		UE::TmvMediaEditor::TranscodeList::FAddedItemsScopedSelectionModifier SelectionModifier(ListHandle, JobListView);
 		TArray<TSharedPtr<FTmvMediaTranscodeJobListItem>> SelectedItems = JobListView->GetSelectedItems();
 		for (const TSharedPtr<FTmvMediaTranscodeJobListItem>& SelectedItem : SelectedItems)
 		{
@@ -859,6 +978,8 @@ void STmvMediaTranscodeList::Paste()
 
 	FScopedTransaction Transaction(LOCTEXT("PasteJobTransaction", "Paste Job"));
 	List->Modify();
+
+	UE::TmvMediaEditor::TranscodeList::FAddedItemsScopedSelectionModifier SelectionModifier(ListHandle, JobListView);
 	for (const FTmvMediaTranscodeListItem& Item : Items)
 	{
 		List->InsertItemAt(List->GetNumItems(), &Item);
