[TmvMedia] Select Apv encoder by default when creating new jobs.

#Synced-CL 51094493

diff --git a/Engine/Plugins/Media/TmvMedia/Source/TmvMedia/Private/Transcoder/TmvMediaTranscodeList.cpp b/Engine/Plugins/Media/TmvMedia/Source/TmvMedia/Private/Transcoder/TmvMediaTranscodeList.cpp
--- a/Engine/Plugins/Media/TmvMedia/Source/TmvMedia/Private/Transcoder/TmvMediaTranscodeList.cpp
+++ b/Engine/Plugins/Media/TmvMedia/Source/TmvMedia/Private/Transcoder/TmvMediaTranscodeList.cpp
@@ -14,17 +14,28 @@ namespace UE::TmvMedia::Transcode
 {
 	const UScriptStruct* GetDefaultEncoderOptionsStruct()
 	{
-		// For now, "default" is whatever the first one is. Should probably be a setting.
+		// Preferred default encoder options is APV.
 		const UScriptStruct* BaseStruct = FTmvMediaEncoderOptions::StaticStruct();
-		for (TObjectIterator<UScriptStruct> It; It; ++It)
+
+		const UScriptStruct* FoundStruct = nullptr;
+		
+		for (const UScriptStruct* CurrentStruct : TObjectRange<UScriptStruct>())
 		{
-			const UScriptStruct* CurrentStruct = *It;
-			if (CurrentStruct->IsChildOf(BaseStruct) && CurrentStruct != BaseStruct)
+			if (CurrentStruct != BaseStruct && CurrentStruct->IsChildOf(BaseStruct))
 			{
-				return CurrentStruct;
+				if (CurrentStruct->GetName().Contains(TEXT("apv")))
+				{
+					FoundStruct = CurrentStruct;
+					break;
+				}
+				if (FoundStruct == nullptr)
+				{
+					FoundStruct = CurrentStruct;
+				}
 			}
 		}
-		return nullptr;
+
+		return FoundStruct;
 	}
 }
 
