[TmvMedia] Fix transcoding toaster message showing the wrong number of processed frames.

It was showing the frame index, which is zero based. The number of processed frame is index + 1.

#Synced-CL 51094493

diff --git a/Engine/Plugins/Media/TmvMedia/Source/TmvMedia/Internal/Transcoder/TmvMediaPlayerFrameProducer.h b/Engine/Plugins/Media/TmvMedia/Source/TmvMedia/Internal/Transcoder/TmvMediaPlayerFrameProducer.h
--- a/Engine/Plugins/Media/TmvMedia/Source/TmvMedia/Internal/Transcoder/TmvMediaPlayerFrameProducer.h
+++ b/Engine/Plugins/Media/TmvMedia/Source/TmvMedia/Internal/Transcoder/TmvMediaPlayerFrameProducer.h
@@ -122,6 +122,12 @@ protected:
 	/** Maximum number of concurrently processing frames. */
 	int32 MaxProcessingFrames = 1;
 
+	/**
+	 * Once the media is completed (MediaCompleted state), this counter keeps track of how many frames still in flight
+	 * we are waiting on before the job is completed (other stages). 
+	 */
+	int32 LastSubmittedCount = 0;
+
 	/** Enum for the frame producer's states. */
 	enum class EState
 	{
diff --git a/Engine/Plugins/Media/TmvMedia/Source/TmvMedia/Private/Transcoder/TmvMediaPlayerFrameProducer.cpp b/Engine/Plugins/Media/TmvMedia/Source/TmvMedia/Private/Transcoder/TmvMediaPlayerFrameProducer.cpp
--- a/Engine/Plugins/Media/TmvMedia/Source/TmvMedia/Private/Transcoder/TmvMediaPlayerFrameProducer.cpp
+++ b/Engine/Plugins/Media/TmvMedia/Source/TmvMedia/Private/Transcoder/TmvMediaPlayerFrameProducer.cpp
@@ -322,25 +322,34 @@ void UTmvMediaPlayerFrameProducer::Tick(UTmvMediaTranscodeJob* InParentJob, cons
 		
 		if (ProducerResource)
 		{
-			// Keep updating the stats
-			const FTimespan RenderedSampleTime = ProducerResource->GetLastRenderedSampleTime();
-			if (RenderedSampleTime != FTimespan::MinValue())
+			const int32 CurrentSubmittedCount = ProducerResource->GetSubmittedCount();
+
+			if (CurrentSubmittedCount != LastSubmittedCount)
 			{
-				const FTimespan RenderedSampleDuration = ProducerResource->GetLastRenderedSampleDuration();
-				const FTimespan FrameDuration = (RenderedSampleDuration > FTimespan::Zero()) ? RenderedSampleDuration : MediaFrameDuration;
-				const FTimespan RelativeTime = RenderedSampleTime - MediaStartTime;
-				if (FrameDuration != FTimespan::Zero() && !MediaPlaybackRange.IsEmpty())
+				LastSubmittedCount = CurrentSubmittedCount;
+
+				// Keep updating the stats
+				const FTimespan RenderedSampleTime = ProducerResource->GetLastRenderedSampleTime();
+				if (RenderedSampleTime != FTimespan::MinValue())
 				{
-					const int32 RenderedFrameIndex = static_cast<int32>(FMath::RoundHalfToEven(RelativeTime.GetTotalSeconds() / FrameDuration.GetTotalSeconds()));
-					const FTimespan SequenceDuration = MediaPlaybackRange.Size<FTimespan>();
-					const int32 TotalFrames = static_cast<int32>(FMath::RoundHalfToEven(SequenceDuration.GetTotalSeconds() / FrameDuration.GetTotalSeconds()));
-					UTmvMediaTranscodeJob::SafeUpdateProgress(InParentJob, RenderedFrameIndex, TotalFrames);
+					const FTimespan RenderedSampleDuration = ProducerResource->GetLastRenderedSampleDuration();
+					const FTimespan FrameDuration = (RenderedSampleDuration > FTimespan::Zero()) ? RenderedSampleDuration : MediaFrameDuration;
+					const FTimespan RelativeTime = RenderedSampleTime - MediaStartTime;
+					if (FrameDuration != FTimespan::Zero() && !MediaPlaybackRange.IsEmpty())
+					{
+						const int32 RenderedFrameIndex = static_cast<int32>(FMath::RoundHalfToEven(RelativeTime.GetTotalSeconds() / FrameDuration.GetTotalSeconds()));
+						const FTimespan SequenceDuration = MediaPlaybackRange.Size<FTimespan>();
+						const int32 TotalFrames = static_cast<int32>(FMath::RoundHalfToEven(SequenceDuration.GetTotalSeconds() / FrameDuration.GetTotalSeconds()));
+						UTmvMediaTranscodeJob::SafeUpdateProgress(InParentJob, FMath::Min(RenderedFrameIndex+1, TotalFrames), TotalFrames);
+					}
 				}
+
+				WatchdogTimeRemaining = 5.0f;	// Reset watchdog.
 			}
-			
+
 			// Wait for final captures to process.
-			ensure(ProducerResource->GetSubmittedCount() >= 0);
-			if (ProducerResource->GetSubmittedCount() == 0)
+			ensure(CurrentSubmittedCount >= 0);
+			if (CurrentSubmittedCount == 0)
 			{
 				RequestStop(InParentJob);
 			}
@@ -434,7 +443,7 @@ void UTmvMediaPlayerFrameProducer::HandlePlaying(UTmvMediaTranscodeJob* InParent
 		if (ProducerResource->GetSubmittedCount() < MaxProcessingFrames)
 		{
 			// Update notification.
-			UTmvMediaTranscodeJob::SafeUpdateProgress(InParentJob, RenderedFrameIndex, TotalFrames);
+			UTmvMediaTranscodeJob::SafeUpdateProgress(InParentJob, FMath::Min(RenderedFrameIndex+1, TotalFrames), TotalFrames);
 
 			// Move to the next frame in the sequence.
 			CurrentTime = RenderedSampleTime + FrameDuration;
@@ -457,6 +466,7 @@ void UTmvMediaPlayerFrameProducer::HandlePlaying(UTmvMediaTranscodeJob* InParent
 					// Frame production is completed on the player side, need to wait for the rest of the stages to finish.
 					CurrentState = EState::MediaCompleted;
 					WatchdogTimeRemaining = 10.0f;
+					LastSubmittedCount = ProducerResource->GetSubmittedCount();
 					// removing the block on range ensures the player can complete without blocking.
 					MediaPlayer->SetBlockOnTimeRange(TRange<FTimespan>::Empty());
 				}
