[TmvMedia] Fix CPU Frame Converter with Alpha channel.

The transfer function (ex: sRGB, PQ, etc) must not be applied to the alpha channel.
This fix only applies to the CPU frame converter code path. The gpu converter was already correct.

#Synced-CL 51094493

diff --git a/Engine/Plugins/Media/TmvMedia/Source/TmvMedia/Private/Transcoder/TmvMediaFrameConverter.cpp b/Engine/Plugins/Media/TmvMedia/Source/TmvMedia/Private/Transcoder/TmvMediaFrameConverter.cpp
--- a/Engine/Plugins/Media/TmvMedia/Source/TmvMedia/Private/Transcoder/TmvMediaFrameConverter.cpp
+++ b/Engine/Plugins/Media/TmvMedia/Source/TmvMedia/Private/Transcoder/TmvMediaFrameConverter.cpp
@@ -176,6 +176,7 @@ namespace UE::TmvMedia::FrameConverter
 			for (int32 PixX = 0; PixX < InImage.SizeX; ++PixX)
 			{
 				FLinearColor RGBA = Row[PixX].GetFloats();
+				const float A = RGBA.A;
 
 #if TMV_MEDIA_USE_VECTOR_PATH
 				VectorRegister4Float vRGBA = VectorLoad(&RGBA.R);
@@ -190,16 +191,14 @@ namespace UE::TmvMedia::FrameConverter
 				RGBA.R = FMath::Max(EncoderType::Encode(RGBA.R), 0.0f);
 				RGBA.G = FMath::Max(EncoderType::Encode(RGBA.G), 0.0f);
 				RGBA.B = FMath::Max(EncoderType::Encode(RGBA.B), 0.0f);
-				RGBA.A = FMath::Max(EncoderType::Encode(RGBA.A), 0.0f);
-
-				float A = RGBA.A;
 				RGBA.A = 1.0f;
+
 				FLinearColor YUVA = RGB2YUVMatrix.TransformFVector4(*((FVector4f*)&RGBA));
-				YUVA.A = A;
 				YUVA.G = FMath::Max(YUVA.G, 0.0f);
 				YUVA.B = FMath::Max(YUVA.B, 0.0f);
 #endif
 
+				YUVA.A = A;
 				Row[PixX] = FFloat16Color(YUVA);
 			}
 		}
