[ApvMedia] Don't issue a warning if the maximum of cores is higher than the max supported by the OAPV implementation.

The number of threads is maxed at 32 in the OAPV implementation. Selecting a number of threads of 0 in ApvMediaSettings would then use the number of cores on the system and issue a warning in the logs if the CPU has more than 32 cores. We had this with a system of 96 cores (Threadripper 9995). The warning will only be issued if the number of threads was explicit.

#Synced-CL 52176027

diff --git a/Engine/Plugins/Media/TmvMedia/Source/ApvMedia/Private/ApvMediaSettings.h b/Engine/Plugins/Media/TmvMedia/Source/ApvMedia/Private/ApvMediaSettings.h
--- a/Engine/Plugins/Media/TmvMedia/Source/ApvMedia/Private/ApvMediaSettings.h
+++ b/Engine/Plugins/Media/TmvMedia/Source/ApvMedia/Private/ApvMediaSettings.h
@@ -19,7 +19,7 @@ public:
 	}
 	
 	/**
-	 * Number of worker thread to decode tiles (0 = automatically use the available number of cores).
+	 * Number of worker thread to decode tiles (0 = automatically use the available number of cores up to a max of 32).
 	 */
 	UPROPERTY(config, EditAnywhere, Category = Decoder)
 	uint32 DecoderThreads = 0;
diff --git a/Engine/Plugins/Media/TmvMedia/Source/ApvMedia/Private/ApvMediaTmvDecoderFactory.cpp b/Engine/Plugins/Media/TmvMedia/Source/ApvMedia/Private/ApvMediaTmvDecoderFactory.cpp
--- a/Engine/Plugins/Media/TmvMedia/Source/ApvMedia/Private/ApvMediaTmvDecoderFactory.cpp
+++ b/Engine/Plugins/Media/TmvMedia/Source/ApvMedia/Private/ApvMediaTmvDecoderFactory.cpp
@@ -49,10 +49,10 @@ void FApvMediaTmvDecoderFactory::GetDecoderOptions(TMap<FString, FVariant>& OutO
 TSharedPtr<ITmvMediaDecoder, ESPMode::ThreadSafe> FApvMediaTmvDecoderFactory::CreateDecoder(const FString& InCodecFormat, const TMap<FString, FVariant>& InOptions)
 {
 	const UApvMediaSettings* Settings = GetDefault<UApvMediaSettings>();
-	int32 NumThreads = (Settings && Settings->DecoderThreads > 0) ? Settings->DecoderThreads : FPlatformMisc::NumberOfCores();
+	int32 NumThreads = (Settings && Settings->DecoderThreads > 0) ? Settings->DecoderThreads : FMath::Min(FPlatformMisc::NumberOfCores(), OAPV_MAX_THREADS);
 	if (NumThreads > OAPV_MAX_THREADS)
 	{
-		UE_LOG(LogApvMedia, Warning, TEXT("Requested number of threads (%d) exceeds maximum OpenAPV threads (%d)."), NumThreads, OAPV_MAX_THREADS);
+		UE_LOG(LogApvMedia, Warning, TEXT("Explicitly requested number of threads (%d) exceeds maximum OpenAPV threads (%d)."), NumThreads, OAPV_MAX_THREADS);
 	}
 	NumThreads = FMath::Clamp(NumThreads, 1, OAPV_MAX_THREADS);
 	return MakeShared<FApvMediaTmvDecoder, ESPMode::ThreadSafe>(NumThreads);
