[TmvMedia] Adding HLG encoding support

- Exposing the HLG encoding in the encoder options
- for Apv, HLG will correspond to ARIB STD-B67.  As long as we encode and decode with our UE implementation, this will be perfect. Can't garantee it will match with other tool's implementations.

#Synced-CL 52257653

diff --git a/Engine/Plugins/Media/TmvMedia/Source/ApvMedia/Private/Apv/ApvMipBufferUtils.h b/Engine/Plugins/Media/TmvMedia/Source/ApvMedia/Private/Apv/ApvMipBufferUtils.h
--- a/Engine/Plugins/Media/TmvMedia/Source/ApvMedia/Private/Apv/ApvMipBufferUtils.h
+++ b/Engine/Plugins/Media/TmvMedia/Source/ApvMedia/Private/Apv/ApvMipBufferUtils.h
@@ -139,6 +139,9 @@ namespace UE::ApvMedia
 
 		//case EEncoding::VLog: // Panasonic V-Log -- unmatched
 
+		case EEncoding::HLG:
+			return EApvMediaColorTransfer::Arib_std_b67;
+
 		default:
 			return TOptional<EApvMediaColorTransfer>();
 		}
@@ -169,6 +172,9 @@ namespace UE::ApvMedia
 		case EApvMediaColorTransfer::Log316:
 			return EEncoding::SLog3;
 
+		case EApvMediaColorTransfer::Arib_std_b67:
+			return EEncoding::HLG;
+
 		default:
 			return EEncoding::sRGB;
 		}
diff --git a/Engine/Plugins/Media/TmvMedia/Source/TmvMedia/Public/Encoder/TmvMediaEncoderOptions.h b/Engine/Plugins/Media/TmvMedia/Source/TmvMedia/Public/Encoder/TmvMediaEncoderOptions.h
--- a/Engine/Plugins/Media/TmvMedia/Source/TmvMedia/Public/Encoder/TmvMediaEncoderOptions.h
+++ b/Engine/Plugins/Media/TmvMedia/Source/TmvMedia/Public/Encoder/TmvMediaEncoderOptions.h
@@ -8,7 +8,7 @@
 #include "TmvMediaEncoderOptions.generated.h"
 
 /** 
- * List of supported encoder transfer functions (a.k.a encodings), matching the list in ColorManagementDefines.h.  
+ * List of supported encoder transfer functions (a.k.a encodings), matching the list in ColorManagementDefines.h.
  */
 UENUM()
 enum class ETmvMediaEncoderEncoding : uint8
@@ -23,6 +23,8 @@ enum class ETmvMediaEncoderEncoding : uint8
 	ST2084 = 3,
 	/** SLog3 encoding */
 	SLog3 = 12,
+	/** Hybrid Log Gamma encoding - Rec. ITU-R BT.2100-3 */
+	HLG	= 17,
 };
 
 /**
@@ -32,7 +34,7 @@ enum class ETmvMediaEncoderEncoding : uint8
 UENUM()
 enum class ETmvMediaEncoderColorMatrix : uint8
 {
-	/** Unspecified. */	
+	/** Unspecified. */
 	None = 0,
 	/** Identity, i.e. no yuv conversion (rgb color model). */
 	Identity = 1,
@@ -51,12 +53,12 @@ UENUM()
 enum class ETmvMediaEncoderColorMatrixRange : uint8
 {
 	/** 
-	 * Limited (studio) range, defined by spec. Typical [16, 235] for luma on 8 bits limited range. 
+	 * Limited (studio) range, defined by spec. Typical [16, 235] for luma on 8 bits limited range.
 	 * This range definition corresponds to the "scaled" yuv matrices and offsets from MediaShaders.
 	 */
 	Limited,
 	/** 
-	 * Full [0, 255] range. 
+	 * Full [0, 255] range.
 	 * This range definition corresponds to the "unscaled" yuv matrices and offsets from MediaShaders.
 	 */
 	Full
@@ -71,7 +73,7 @@ struct FTmvMediaEncoderOptions
 	GENERATED_BODY()
 
 	/** Returns the encoder name used to find the corresponding factory. */
-	virtual FName GetEncoderName() const { return NAME_None;}
+	virtual FName GetEncoderName() const { return NAME_None; }
 
 	/** Returns the extension used for saving to file sequence. */
 	virtual FString GetFileSequenceExtension() const { return FString(); }
@@ -79,7 +81,7 @@ struct FTmvMediaEncoderOptions
 	virtual ~FTmvMediaEncoderOptions() = default;
 
 	/** 
-	 * Enables the color management specification for the encoder. 
+	 * Enables the color management specification for the encoder.
 	 */
 	UPROPERTY(EditAnywhere, Category = "Color Management")
 	bool bEnableColorManagement = false;
@@ -91,11 +93,11 @@ struct FTmvMediaEncoderOptions
 	/** Specify the destination color encoding. */
 	UPROPERTY(EditAnywhere, Category = "Color Management", meta = (EditCondition = "bEnableColorManagement", EditConditionHides))
 	ETmvMediaEncoderEncoding DestinationEncoding = ETmvMediaEncoderEncoding::Linear;
-	
+
 	/** Specify the YUV encoding color matrix. */
 	UPROPERTY(EditAnywhere, Category = "Color Management", meta = (EditCondition = "bEnableColorManagement", EditConditionHides))
 	ETmvMediaEncoderColorMatrix YuvMatrix = ETmvMediaEncoderColorMatrix::Identity;
-	
+
 	/** Specify the YUV encoding color matrix range. */
 	UPROPERTY(EditAnywhere, Category = "Color Management", meta = (EditCondition = "bEnableColorManagement", EditConditionHides))
 	ETmvMediaEncoderColorMatrixRange YuvMatrixRange = ETmvMediaEncoderColorMatrixRange::Full;
