[TmvMedia] Transcoder: Fix JobItem name collision on insert

#Synced-CL 52257653

diff --git a/Engine/Plugins/Media/TmvMedia/Source/TmvMedia/Private/Transcoder/TmvMediaTranscodeList.cpp b/Engine/Plugins/Media/TmvMedia/Source/TmvMedia/Private/Transcoder/TmvMediaTranscodeList.cpp
--- a/Engine/Plugins/Media/TmvMedia/Source/TmvMedia/Private/Transcoder/TmvMediaTranscodeList.cpp
+++ b/Engine/Plugins/Media/TmvMedia/Source/TmvMedia/Private/Transcoder/TmvMediaTranscodeList.cpp
@@ -41,29 +41,28 @@ namespace UE::TmvMedia::Transcode
 
 void UTmvMediaTranscodeList::InsertItemAt(int32 InItemIndex, const FTmvMediaTranscodeListItem* InItem)
 {
-	// Invalid index.
 	if (InItemIndex < 0)
 	{
 		return;
 	}
-	
+
 	FTmvMediaTranscodeListItem NewItem;
+	NewItem.Id = FGuid::NewGuid();
+
 	if (InItem)
 	{
-		NewItem.Name = InItem->Name + TEXT("_Dup");
+		NewItem.Name = GenerateDuplicateName(InItem->Name);
 		NewItem.Settings = InItem->Settings;
 		NewItem.EncoderOptions = InItem->EncoderOptions;
 	}
 	else
 	{
-		NewItem.Name = FString::Printf(TEXT("JobItem_%d"), InItemIndex);
+		NewItem.Name = GenerateNewItemName();
 		if (const UScriptStruct* SelectedStruct = UE::TmvMedia::Transcode::GetDefaultEncoderOptionsStruct())
 		{
 			NewItem.EncoderOptions.InitializeAsScriptStruct(SelectedStruct);
 		}
 	}
-	
-	NewItem.Id = FGuid::NewGuid();
 
 	if (InItemIndex >= Items.Num())
 	{
@@ -183,4 +182,73 @@ void UTmvMediaTranscodeList::PostEditUndo()
 }
 #endif
 
+FString UTmvMediaTranscodeList::GenerateNewItemName() const
+{
+	// Find the highest existing JobItem_N number to avoid collisions when items have been deleted.
+	static const FString Prefix = TEXT("JobItem_");
+	int32 MaxNumber = -1;
+	for (const FTmvMediaTranscodeListItem& Item : Items)
+	{
+		if (Item.Name.StartsWith(Prefix))
+		{
+			int32 Number;
+			if (LexTryParseString(Number, *Item.Name.RightChop(Prefix.Len())))
+			{
+				MaxNumber = FMath::Max(MaxNumber, Number);
+			}
+		}
+	}
+	return FString::Printf(TEXT("JobItem_%d"), MaxNumber + 1);
+}
+
+FString UTmvMediaTranscodeList::GenerateDuplicateName(const FString& InSourceName) const
+{
+	static const FString DupSuffix = TEXT("_Dup");
+
+	// Strip any existing _Dup or _DupN suffix to get the base name.
+	FString BaseName = InSourceName;
+	if (BaseName.EndsWith(DupSuffix))
+	{
+		BaseName.LeftChopInline(DupSuffix.Len());
+	}
+	else
+	{
+		const int32 LastDupIdx = BaseName.Find(DupSuffix, ESearchCase::CaseSensitive, ESearchDir::FromEnd);
+		if (LastDupIdx != INDEX_NONE)
+		{
+			const FString Tail = BaseName.RightChop(LastDupIdx + DupSuffix.Len());
+			int32 Num;
+			if (LexTryParseString(Num, *Tail))
+			{
+				BaseName.LeftInline(LastDupIdx);
+			}
+		}
+	}
+
+	// Find the highest existing dup index for this base name.
+	// _Dup counts as index 1, _Dup2 as 2, etc.
+	const FString DupPrefix = BaseName + DupSuffix;
+	int32 MaxDupIndex = 0;
+	for (const FTmvMediaTranscodeListItem& ExistingItem : Items)
+	{
+		if (ExistingItem.Name == DupPrefix)
+		{
+			MaxDupIndex = FMath::Max(MaxDupIndex, 1);
+		}
+		else if (ExistingItem.Name.StartsWith(DupPrefix))
+		{
+			const FString Tail = ExistingItem.Name.RightChop(DupPrefix.Len());
+			int32 Num;
+			if (LexTryParseString(Num, *Tail) && Num >= 2)
+			{
+				MaxDupIndex = FMath::Max(MaxDupIndex, Num);
+			}
+		}
+	}
+
+	return MaxDupIndex == 0
+		? DupPrefix
+		: FString::Printf(TEXT("%s_Dup%d"), *BaseName, MaxDupIndex + 1);
+}
+
 
diff --git a/Engine/Plugins/Media/TmvMedia/Source/TmvMedia/Public/Transcoder/TmvMediaTranscodeList.h b/Engine/Plugins/Media/TmvMedia/Source/TmvMedia/Public/Transcoder/TmvMediaTranscodeList.h
--- a/Engine/Plugins/Media/TmvMedia/Source/TmvMedia/Public/Transcoder/TmvMediaTranscodeList.h
+++ b/Engine/Plugins/Media/TmvMedia/Source/TmvMedia/Public/Transcoder/TmvMediaTranscodeList.h
@@ -134,10 +134,16 @@ public:
 	//~ End UObject
 
 private:
+	/** Generates a unique new item name of the form "JobItem_N". */
+	FString GenerateNewItemName() const;
+
+	/** Generates a unique duplicate name for the given source name, e.g. "Foo_Dup", "Foo_Dup2". */
+	FString GenerateDuplicateName(const FString& InSourceName) const;
+
 	/** Main list of all job items. */
 	UPROPERTY()
 	TArray<FTmvMediaTranscodeListItem> Items;
-	
+
 	/** Delegate for item editing events. */
 	FOnItemEvent OnItemEvent;
 };
