[TmvMediaEditor] Fix Job parameters panel does not refresh when switching jobs after editing text properties.

Root Cause:
Editing a path property leaves the SEditableTextBox focused. When the user clicks a
different job row, Slate's deferred selection mode (ETableRowSignalSelectionMode::Deferred)
queues a SetUserFocus reply on mouse-DOWN before signaling the selection change on
mouse-UP. The focus change causes the text box to commit its text synchronously, which
fires OnFinishedChangingProperties. This broadcasts an ItemsModified event, which called
RebuildList() on the job list view. The rebuild destroyed and recreated the row widget
for the job being clicked mid-click, resetting its bChangedSelectionOnMouseDown flag to
false. When mouse-UP arrived on the new widget, the guard condition was false so
Private_SignalSelectionChanged was never called, and the panel never refreshed.

Fix:
Skip RebuildList() in STmvMediaTranscodeList::OnTranscodeListItemEvent for ItemsModified
events. A full rebuild is only needed when the item count changes (add/remove/reorder).
For modifications, the row widgets' attribute bindings already poll the updated cached
values on the next paint, so no rebuild is necessary.

#Synced-CL 52257653

diff --git a/Engine/Plugins/Media/TmvMedia/Source/TmvMediaEditor/Private/Widgets/STmvMediaTranscodeList.cpp b/Engine/Plugins/Media/TmvMedia/Source/TmvMediaEditor/Private/Widgets/STmvMediaTranscodeList.cpp
--- a/Engine/Plugins/Media/TmvMedia/Source/TmvMediaEditor/Private/Widgets/STmvMediaTranscodeList.cpp
+++ b/Engine/Plugins/Media/TmvMedia/Source/TmvMediaEditor/Private/Widgets/STmvMediaTranscodeList.cpp
@@ -992,11 +992,18 @@ void STmvMediaTranscodeList::OnTranscodeListItemEvent(const UTmvMediaTranscodeLi
 	UTmvMediaTranscodeList* CurrentList = ListHandle ? ListHandle->Get() : nullptr;
 	if (CurrentList && CurrentList == InList)
 	{
-		// todo: For item modified event, just refresh the modified item.
 		FTmvMediaTranscodeJobListItem::RefreshList(JobList, CurrentList);
 		if (JobListView)
 		{
-			JobListView->RebuildList();
+			// For ItemsModified, skip RebuildList() — it destroys and recreates row widgets mid-click,
+			// which clears the bChangedSelectionOnMouseDown flag on the clicked row. This causes the
+			// subsequent mouse-up to skip signaling the selection change, so the parameters panel
+			// never refreshes when switching jobs. Row attribute bindings pick up the updated cached
+			// values on the next paint without requiring a full rebuild.
+			if (InArgs.Type != ETmvMediaTranscodeListItemEventType::ItemsModified)
+			{
+				JobListView->RebuildList();
+			}
 			if (InArgs.Type == ETmvMediaTranscodeListItemEventType::ItemsRemoved || InArgs.Type == ETmvMediaTranscodeListItemEventType::ItemsReordered)
 			{
 				JobListView->ClearSelection();
