[ApvMedia] ApvTmvEncoder: Fix tile sizes to be aligned with macroblock size (encoder constraint).

Also adding error messages for the toaster message in case of encoder creation failure.

#Synced-CL 52257653

diff --git a/Engine/Plugins/Media/TmvMedia/Source/ApvMedia/Private/Apv/ApvEncoder.cpp b/Engine/Plugins/Media/TmvMedia/Source/ApvMedia/Private/Apv/ApvEncoder.cpp
--- a/Engine/Plugins/Media/TmvMedia/Source/ApvMedia/Private/Apv/ApvEncoder.cpp
+++ b/Engine/Plugins/Media/TmvMedia/Source/ApvMedia/Private/Apv/ApvEncoder.cpp
@@ -7,8 +7,11 @@
 #include "ApvMipBufferUtils.h"
 #include "Encoder/ITmvMediaEncoder.h"
 #include "Transcoder/TmvMediaFrameMips.h"
+#include "Utils/TmvMediaMessageContext.h"
 #include "Utils/TmvMediaUtils.h"
 
+#define LOCTEXT_NAMESPACE "ApvEncoderContext"
+
 namespace UE::ApvMedia
 {
 	int32 GetProfileIdc(EApvMediaProfile InProfile)
@@ -73,7 +76,8 @@ namespace UE::ApvMedia
 	FApvEncoderContext::FApvEncoderContext(
 		const FApvMediaTmvEncoderOptions& InOptions,
 		const FTmvMediaEncoderMipInfo& InMipInfo,
-		const FTmvMediaFrameTimeInfo& InTimeInfo)
+		const FTmvMediaFrameTimeInfo& InTimeInfo,
+		FTmvMediaMessageContext* OutMessageContext)
 	{
 		// Remark:
 		// we need to have all the encoder options when creating the encoder.
@@ -89,15 +93,18 @@ namespace UE::ApvMedia
 		// Report invalid frame dimensions.
 		if (InMipInfo.Width <= 0 || InMipInfo.Height <= 0)
 		{
-			UE_LOG(LogApvMedia, Error, TEXT("ApvEncoderContext: invalid frame dimensions %dx%d."), InMipInfo.Width, InMipInfo.Height);
+			FText Message = FText::Format(
+				LOCTEXT("InvalidFrameDim", "Invalid frame dimensions ({0} x {1})"), 
+				FText::AsNumber(InMipInfo.Width), FText::AsNumber(InMipInfo.Height));
+			UE_TMV_MEDIA_MESSAGE_LOG(OutMessageContext, LogApvMedia, Error, TEXT("ApvEncoderContext"), Message);
 			return;
 		}
 
 		// Setup encoder options here.
 		PrimaryFrameParams.w = InMipInfo.Width;
 		PrimaryFrameParams.h = InMipInfo.Height;
-		PrimaryFrameParams.tile_w = InOptions.TileSize.X;
-		PrimaryFrameParams.tile_h = InOptions.TileSize.Y;
+		PrimaryFrameParams.tile_w = Align(InOptions.TileSize.X, OAPV_MB_W);
+		PrimaryFrameParams.tile_h = Align(InOptions.TileSize.Y, OAPV_MB_H);
 		PrimaryFrameParams.profile_idc = GetProfileIdc(InOptions.Profile);
 		PrimaryFrameParams.band_idc = FMath::Clamp(InOptions.Band, 0, 3);
 		PrimaryFrameParams.preset = GetPreset(InOptions.Preset);
@@ -193,7 +200,9 @@ namespace UE::ApvMedia
 		// Report invalid bitstream buffer size.
 		if (EncoderDescriptor.max_bs_buf_size <= 0)
 		{
-			UE_LOG(LogApvMedia, Error, TEXT("ApvEncoderContext: Invalid MaxEncoderBitStreamBufferSize (%d). Recommended 128 MB."), EncoderDescriptor.max_bs_buf_size);
+			FText Message = FText::Format(
+				LOCTEXT("InvalidBitstreamSize", "Invalid MaxEncoderBitStreamBufferSize {0}. Recommended 128 MB"), FText::AsNumber(EncoderDescriptor.max_bs_buf_size));
+			UE_TMV_MEDIA_MESSAGE_LOG(OutMessageContext, LogApvMedia, Error, TEXT("ApvEncoderContext"), Message);
 			return;
 		}
 
@@ -201,7 +210,10 @@ namespace UE::ApvMedia
 		eid = oapve_create(&EncoderDescriptor, &Err);
 		if (eid == nullptr)
 		{
-			UE_LOG(LogApvMedia, Error, TEXT("ApvEncoderContext: cannot create OpenApv encoder (err=%d)"), Err);
+			FText Message = FText::Format(
+				LOCTEXT("FailedCreateEncoder", "Cannot create OpenApv encoder, error: {0} ({1})"), 
+				FText::FromString(GetApvErrorString(Err)), FText::AsNumber(Err));
+			UE_TMV_MEDIA_MESSAGE_LOG(OutMessageContext, LogApvMedia, Error, TEXT("ApvEncoderContext"), Message);
 		}
 		else
 		{
@@ -213,7 +225,10 @@ namespace UE::ApvMedia
 		mid = oapvm_create(&Err);
 		if(OAPV_FAILED(Err))
 		{
-			UE_LOG(LogApvMedia, Error, TEXT("ApvEncoderContext: cannot create OpenApv metadata container (err=%d)"), Err);
+			FText Message = FText::Format(
+				LOCTEXT("FailedCreateMetadata", "Cannot create OpenApv metadata container, error: {0} ({1})"), 
+				FText::FromString(GetApvErrorString(Err)), FText::AsNumber(Err));
+			UE_TMV_MEDIA_MESSAGE_LOG(OutMessageContext, LogApvMedia, Error, TEXT("ApvEncoderContext"), Message);
 		}
 	}
 
@@ -231,3 +246,5 @@ namespace UE::ApvMedia
 		}
 	}
 }
+
+#undef LOCTEXT_NAMESPACE
\ No newline at end of file
diff --git a/Engine/Plugins/Media/TmvMedia/Source/ApvMedia/Private/Apv/ApvEncoder.h b/Engine/Plugins/Media/TmvMedia/Source/ApvMedia/Private/Apv/ApvEncoder.h
--- a/Engine/Plugins/Media/TmvMedia/Source/ApvMedia/Private/Apv/ApvEncoder.h
+++ b/Engine/Plugins/Media/TmvMedia/Source/ApvMedia/Private/Apv/ApvEncoder.h
@@ -5,6 +5,7 @@
 #include "ApvCommon.h"
 #include "ApvDecoder.h"
 
+class FTmvMediaMessageContext;
 struct FApvMediaTmvEncoderOptions;
 struct FTmvMediaEncoderMipInfo;
 struct FTmvMediaFrameTimeInfo;
@@ -21,12 +22,14 @@ namespace UE::ApvMedia
 		 * @param InOptions Encoder options (user specified)
 		 * @param InMipInfo Primary Mip (mip 0) descriptor
 		 * @param InTimeInfo Only Time info's frame rate is used as a hint for bit rate calculations.
+		 * @param OutMessageContext Optional message context for the given operation.
 		 * @remark Sequence's frame rate can't be stored by OpenApv frame header (not supported). 
 		 */
 		FApvEncoderContext(
 			const FApvMediaTmvEncoderOptions& InOptions,
 			const FTmvMediaEncoderMipInfo& InMipInfo,
-			const FTmvMediaFrameTimeInfo& InTimeInfo);
+			const FTmvMediaFrameTimeInfo& InTimeInfo,
+			FTmvMediaMessageContext* OutMessageContext);
 
 		// Non-copiable
 		FApvEncoderContext(const FApvEncoderContext& InOther) = delete;
diff --git a/Engine/Plugins/Media/TmvMedia/Source/ApvMedia/Private/ApvMediaTmvEncoder.cpp b/Engine/Plugins/Media/TmvMedia/Source/ApvMedia/Private/ApvMediaTmvEncoder.cpp
--- a/Engine/Plugins/Media/TmvMedia/Source/ApvMedia/Private/ApvMediaTmvEncoder.cpp
+++ b/Engine/Plugins/Media/TmvMedia/Source/ApvMedia/Private/ApvMediaTmvEncoder.cpp
@@ -170,7 +170,7 @@ ETmvMediaEncoderResult FApvMediaTmvEncoder::Encode(
 			EncoderMipInfo.Height = MipBuffer->GetMipInfoRef().Height;
 			EncoderMipInfo.ColorInfo = MipBuffer->GetMipInfoRef().ColorInfo;
 			EncoderMipInfo.bEnableMips = InMipRequests.Num() > 1 ? true : false;
-			EncoderContext = MakeUnique<UE::ApvMedia::FApvEncoderContext>(EncoderOptions, EncoderMipInfo, InTimeInfo);
+			EncoderContext = MakeUnique<UE::ApvMedia::FApvEncoderContext>(EncoderOptions, EncoderMipInfo, InTimeInfo, OutMessageContext);
 		}
 		else 
 		{
