[TmvMedia] Transcoder: Create or Update Media Source asset for the converted sequence

A destination content path must be specified and the job name is used as the destination media source asset name.
The desination media source will have it's frame rate set from the source media or job settings.

Media Source Make or Update:
- supports img sequences, but is ready to support containers too.
- will only save the media source if it is modified to avoid triggereing source control transaction if not needed.

Also, adding an explanatory tool tip to the "start jobs" button, in case it is disabled, to better guide the user as to what is not configured correctly.

#Synced-CL 52257653

diff --git a/Engine/Plugins/Media/TmvMedia/Source/TmvMedia/Internal/Transcoder/TmvMediaTmvFrameEncoder.h b/Engine/Plugins/Media/TmvMedia/Source/TmvMedia/Internal/Transcoder/TmvMediaTmvFrameEncoder.h
--- a/Engine/Plugins/Media/TmvMedia/Source/TmvMedia/Internal/Transcoder/TmvMediaTmvFrameEncoder.h
+++ b/Engine/Plugins/Media/TmvMedia/Source/TmvMedia/Internal/Transcoder/TmvMediaTmvFrameEncoder.h
@@ -24,6 +24,7 @@ public:
 
 	//~ Begin UTmvMediaFrameEncoder
 	virtual void SetEncoderOptions(const TInstancedStruct<FTmvMediaEncoderOptions>& InOptions) override;
+	virtual const FTmvMediaEncoderOptions* GetEncoderOptions() const override;
 	virtual bool RequestMipInfos(UTmvMediaTranscodeJob* InParentJob, const FTmvMediaFrameTimeInfo& InSourceTimeInfo, const FTmvMediaFrameMipInfo& InSourceMipInfo, TArray<FTmvMediaFrameMipInfo>& OutFrameMipInfo) override;
 	virtual void ReceiveMips(UTmvMediaTranscodeJob* InParentJob, TUniquePtr<FTmvMediaFrameMips>&& InMips) override;
 	//~ End UTmvMediaFrameEncoder
diff --git a/Engine/Plugins/Media/TmvMedia/Source/TmvMedia/Internal/Utils/TmvMediaPathUtils.h b/Engine/Plugins/Media/TmvMedia/Source/TmvMedia/Internal/Utils/TmvMediaPathUtils.h
--- a/Engine/Plugins/Media/TmvMedia/Source/TmvMedia/Internal/Utils/TmvMediaPathUtils.h
+++ b/Engine/Plugins/Media/TmvMedia/Source/TmvMedia/Internal/Utils/TmvMediaPathUtils.h
@@ -4,6 +4,8 @@
 
 #include "Containers/UnrealString.h"
 
+class UMediaSource;
+
 #define UE_API TMVMEDIA_API
 
 /**
@@ -58,6 +60,13 @@ namespace UE::TmvMedia::PathUtils
 	 * @return formatted filename with frame index and extension appended. 
 	 */
 	UE_API FString MakeFrameFilename(const FString& InBaseName, int32 InFrameIndex, int32 InZeroPadFrameNumbers, const FString& InExtension);
+
+	/**
+	 * Utility function to retrieve the path that a media source is pointing to.
+	 * @param InMediaSource Media source
+	 * @return cleaned up path to an existing file or folder, or empty string if not supported by the media source.
+	 */
+	UE_API FString GetMediaSourceMediaFullPath(const UMediaSource* InMediaSource);
 }
 
 #undef UE_API
diff --git a/Engine/Plugins/Media/TmvMedia/Source/TmvMedia/Private/Transcoder/TmvMediaFileSequenceFrameProducer.cpp b/Engine/Plugins/Media/TmvMedia/Source/TmvMedia/Private/Transcoder/TmvMediaFileSequenceFrameProducer.cpp
--- a/Engine/Plugins/Media/TmvMedia/Source/TmvMedia/Private/Transcoder/TmvMediaFileSequenceFrameProducer.cpp
+++ b/Engine/Plugins/Media/TmvMedia/Source/TmvMedia/Private/Transcoder/TmvMediaFileSequenceFrameProducer.cpp
@@ -163,6 +163,15 @@ void UTmvMediaFileSequenceFrameProducer::ProcessAllImages(UTmvMediaTranscodeJob*
 		return;
 	}
 
+	// Update the video track info
+	if (InParentJob->Settings.FrameRate.Numerator > 0 && InParentJob->Settings.FrameRate.Denominator > 0)
+	{
+		FTmvMediaFrameProducerTrackInfo LocalVideoTrackInfo = GetVideoTrackInfo();
+		LocalVideoTrackInfo.FrameRate = InParentJob->Settings.FrameRate;
+		LocalVideoTrackInfo.Duration = FTimespan::FromSeconds(LocalVideoTrackInfo.FrameRate.AsInterval() * static_cast<double>(FoundFiles.Num()));
+		SetVideoTrackInfo(LocalVideoTrackInfo);
+	}
+
 	// Get number of threads to use.
 	int32 NumThreads = InParentJob->Settings.NumProducerThreads;
 	if (NumThreads <= 0)
diff --git a/Engine/Plugins/Media/TmvMedia/Source/TmvMedia/Private/Transcoder/TmvMediaFrameProducer.cpp b/Engine/Plugins/Media/TmvMedia/Source/TmvMedia/Private/Transcoder/TmvMediaFrameProducer.cpp
new file mode 100644
--- /dev/null
+++ b/Engine/Plugins/Media/TmvMedia/Source/TmvMedia/Private/Transcoder/TmvMediaFrameProducer.cpp
@@ -0,0 +1,17 @@
+// Copyright Epic Games, Inc. All Rights Reserved.
+
+#include "Transcoder/TmvMediaFrameProducer.h"
+
+#include UE_INLINE_GENERATED_CPP_BY_NAME(TmvMediaFrameProducer)
+
+FTmvMediaFrameProducerTrackInfo UTmvMediaFrameProducer::GetVideoTrackInfo()	const
+{
+	FScopeLock lock(&TrackInfoLock);
+	return 	VideoTrackInfo;
+}
+
+void UTmvMediaFrameProducer::SetVideoTrackInfo(const FTmvMediaFrameProducerTrackInfo& InVideoTrackInfo)
+{
+	FScopeLock lock(&TrackInfoLock);
+	VideoTrackInfo = InVideoTrackInfo;
+}
diff --git a/Engine/Plugins/Media/TmvMedia/Source/TmvMedia/Private/Transcoder/TmvMediaPlayerFrameProducer.cpp b/Engine/Plugins/Media/TmvMedia/Source/TmvMedia/Private/Transcoder/TmvMediaPlayerFrameProducer.cpp
--- a/Engine/Plugins/Media/TmvMedia/Source/TmvMedia/Private/Transcoder/TmvMediaPlayerFrameProducer.cpp
+++ b/Engine/Plugins/Media/TmvMedia/Source/TmvMedia/Private/Transcoder/TmvMediaPlayerFrameProducer.cpp
@@ -60,7 +60,7 @@ bool UTmvMediaPlayerFrameProducer::Start(UTmvMediaTranscodeJob* InParentJob)
 		UE_LOG(LogTmvMedia, Error, TEXT("UTmvMediaPlayerFrameProducer::Start: InParentJob is null."));
 		return false;
 	}
-	
+
 	MaxProcessingFrames = InParentJob->Settings.NumProducerThreads;
 	if (MaxProcessingFrames <= 0)
 	{
@@ -539,6 +539,11 @@ void UTmvMediaPlayerFrameProducer::TryUpdatePlaybackRange()
 			UE_LOG(LogTmvMedia, Error, TEXT("Transcode Frame Producer: Media has no video track."));
 			CurrentState = EState::NotSupported;
 		}
+		
+        // Update the track info
+        FTmvMediaFrameProducerTrackInfo LocalVideoTrackInfo = GetVideoTrackInfo();
+		LocalVideoTrackInfo.Duration =  MediaPlaybackRange.Size<FTimespan>();
+        SetVideoTrackInfo(LocalVideoTrackInfo);
 	}
 }
 
@@ -551,6 +556,21 @@ void UTmvMediaPlayerFrameProducer::TryUpdateFrameDuration()
 	{
 		UE_LOG(LogTmvMedia, Verbose, TEXT("Transcode Frame Producer: Media Frame Rate is %f."), FrameRate);
 		MediaFrameDuration = FTimespan::FromSeconds(1.0f / FrameRate);
+		
+        // Update the track info
+        FTmvMediaFrameProducerTrackInfo LocalVideoTrackInfo = GetVideoTrackInfo();
+        
+        // Fractional frame rates are considered "drop frame"
+        if (FMath::Fractional(FrameRate) > UE_KINDA_SMALL_NUMBER)
+        {
+        	LocalVideoTrackInfo.FrameRate = FFrameRate(FMath::CeilToInt(FrameRate) * 1000, 1001);
+        }
+        else
+        {
+        	LocalVideoTrackInfo.FrameRate = FFrameRate(FrameRate, 1);
+        }
+
+        SetVideoTrackInfo(LocalVideoTrackInfo);
 	}
 }
 
diff --git a/Engine/Plugins/Media/TmvMedia/Source/TmvMedia/Private/Transcoder/TmvMediaTmvFrameEncoder.cpp b/Engine/Plugins/Media/TmvMedia/Source/TmvMedia/Private/Transcoder/TmvMediaTmvFrameEncoder.cpp
--- a/Engine/Plugins/Media/TmvMedia/Source/TmvMedia/Private/Transcoder/TmvMediaTmvFrameEncoder.cpp
+++ b/Engine/Plugins/Media/TmvMedia/Source/TmvMedia/Private/Transcoder/TmvMediaTmvFrameEncoder.cpp
@@ -214,6 +214,11 @@ void UTmvMediaTmvFrameEncoder::SetEncoderOptions(const TInstancedStruct<FTmvMedi
 	EncoderOptions = InOptions;
 }
 
+const FTmvMediaEncoderOptions* UTmvMediaTmvFrameEncoder::GetEncoderOptions() const
+{
+	return EncoderOptions.GetPtr();
+}
+
 bool UTmvMediaTmvFrameEncoder::RequestMipInfos(UTmvMediaTranscodeJob* InParentJob, const FTmvMediaFrameTimeInfo& InSourceTimeInfo, const FTmvMediaFrameMipInfo& InSourceMipInfo, TArray<FTmvMediaFrameMipInfo>& OutFrameMipInfo)
 {
 	if (!InParentJob)
diff --git a/Engine/Plugins/Media/TmvMedia/Source/TmvMedia/Private/Transcoder/TmvMediaTranscodeJob.cpp b/Engine/Plugins/Media/TmvMedia/Source/TmvMedia/Private/Transcoder/TmvMediaTranscodeJob.cpp
--- a/Engine/Plugins/Media/TmvMedia/Source/TmvMedia/Private/Transcoder/TmvMediaTranscodeJob.cpp
+++ b/Engine/Plugins/Media/TmvMedia/Source/TmvMedia/Private/Transcoder/TmvMediaTranscodeJob.cpp
@@ -35,30 +35,17 @@ FString FTmvMediaTranscodeJobSettings::GetAbsoluteInputPath() const
 	// Get the media file from the media source.
 	if (UMediaSource* MediaSource = InputMediaSource.LoadSynchronous())
 	{
-		FString Url = MediaSource->GetUrl();
+		const FString MediaPath = UE::TmvMedia::PathUtils::GetMediaSourceMediaFullPath(MediaSource);
 
-		// Clean up to get just the filename.
-		const FString FilePrefix(TEXT("file://"));
-		const FString ImgPrefix(TEXT("img://"));
-
-		if (Url.StartsWith(FilePrefix))
-		{
-			return UE::TmvMedia::PathUtils::ConvertSanitizedPathToFull(Url.RightChop(FilePrefix.Len()));
-		}
-		
-		if (Url.StartsWith(ImgPrefix))
+		if (MediaPath.IsEmpty())
 		{
-			// Img Media only keeps the directory. For image sequence producer, it will then search for files in that directory.
-			return UE::TmvMedia::PathUtils::ConvertSanitizedPathToFull(Url.RightChop(ImgPrefix.Len()));
+			UE_LOG(LogTmvMedia, Error, TEXT("Media Source \"%s\" is not a supported codec for transcoding."), *MediaSource->GetUrl());
 		}
 
-		UE_LOG(LogTmvMedia, Error, TEXT("Media Source \"%s\" is not a supported codec for transcoding."), *Url);
-	}
-	else
-	{
-		UE_LOG(LogTmvMedia, Error, TEXT("Unable to load asset \"%s\"."), *InputMediaSource.ToString());
+		return MediaPath;
 	}
 
+	UE_LOG(LogTmvMedia, Error, TEXT("Unable to load asset \"%s\"."), *InputMediaSource.ToString());
 	return FString();
 }
 
@@ -73,6 +60,32 @@ bool FTmvMediaTranscodeJobSettings::IsInputPathSet() const
 	return InputMediaSource.IsNull() == false;
 }
 
+FFrameRate FTmvMediaTranscodeJobSettings::GetInputFramerate() const
+{
+	if (InputSource == ETmvMediaTranscodeInputSource::MediaSource)
+	{
+		// Get the frame rate from the media source if possible.
+		if (UMediaSource* MediaSource = InputMediaSource.LoadSynchronous())
+		{
+			if (FStructProperty* const FrameRateProperty = FindFProperty<FStructProperty>(MediaSource->GetClass(), TEXT("FrameRateOverride")))
+			{
+				if (const void* SrcPtr = FrameRateProperty->ContainerPtrToValuePtr<FFrameRate>(MediaSource, 0))
+				{
+					FFrameRate LocalFrameRate = FrameRate;
+					FrameRateProperty->CopyCompleteValue(&LocalFrameRate, SrcPtr);
+					return LocalFrameRate;
+				}
+			}
+		}
+		else
+		{
+			UE_LOG(LogTmvMedia, Error, TEXT("Unable to load asset \"%s\"."), *InputMediaSource.ToString());
+		}
+	}
+
+	return FrameRate;
+}
+
 bool FTmvMediaTranscodeJobSettings::IsOutputPathSet() const
 {
 	return OutputPath.Path.IsEmpty() == false;
@@ -184,8 +197,14 @@ void UTmvMediaTranscodeJob::RequestStop(double InCurrentTime, ETmvMediaTranscode
 	// Check if the stages are actually stopped.
 	if (AreStagesStopped())
 	{
+		ETmvMediaTranscodeJobStatus PreviousJobStatus = JobStatus;
 		JobStatus = ETmvMediaTranscodeJobStatus::Stopped;	// truly stopped.
-		OnStopped(InCurrentTime);
+		
+		// Only call OnStopped if the job actually stopped.
+		if (PreviousJobStatus != JobStatus)
+		{
+			OnStopped(InCurrentTime);
+		}
 	}
 }
 
diff --git a/Engine/Plugins/Media/TmvMedia/Source/TmvMedia/Private/Utils/TmvMediaPathUtils.cpp b/Engine/Plugins/Media/TmvMedia/Source/TmvMedia/Private/Utils/TmvMediaPathUtils.cpp
--- a/Engine/Plugins/Media/TmvMedia/Source/TmvMedia/Private/Utils/TmvMediaPathUtils.cpp
+++ b/Engine/Plugins/Media/TmvMedia/Source/TmvMedia/Private/Utils/TmvMediaPathUtils.cpp
@@ -2,6 +2,7 @@
 
 #include "Utils/TmvMediaPathUtils.h"
 
+#include "MediaSource.h"
 #include "Misc/Optional.h"
 #include "Misc/Paths.h"
 #include "Templates/Function.h"
@@ -262,4 +263,28 @@ namespace UE::TmvMedia::PathUtils
 		
 		return FPaths::MakeValidFileName(FormattedName);
 	}
+	
+	FString GetMediaSourceMediaFullPath(const UMediaSource* InMediaSource)
+	{
+		if (InMediaSource)
+		{
+			// For ImgMediaSource and FileMediaSource, the url will already return the full path.
+			const FString Url = InMediaSource->GetUrl();
+
+			// Clean up to get just the path.
+			const FString FilePrefix(TEXT("file://"));
+			const FString ImgPrefix(TEXT("img://"));
+
+			if (Url.StartsWith(FilePrefix))
+			{
+				return ConvertSanitizedPathToFull(Url.RightChop(FilePrefix.Len()));
+			}
+		
+			if (Url.StartsWith(ImgPrefix))
+			{
+				return ConvertSanitizedPathToFull(Url.RightChop(ImgPrefix.Len()));
+			}
+		}
+		return FString();
+	}
 }
\ No newline at end of file
diff --git a/Engine/Plugins/Media/TmvMedia/Source/TmvMedia/Public/Transcoder/TmvMediaFrameEncoder.h b/Engine/Plugins/Media/TmvMedia/Source/TmvMedia/Public/Transcoder/TmvMediaFrameEncoder.h
--- a/Engine/Plugins/Media/TmvMedia/Source/TmvMedia/Public/Transcoder/TmvMediaFrameEncoder.h
+++ b/Engine/Plugins/Media/TmvMedia/Source/TmvMedia/Public/Transcoder/TmvMediaFrameEncoder.h
@@ -29,6 +29,11 @@ public:
 	 */
 	virtual void SetEncoderOptions(const TInstancedStruct<FTmvMediaEncoderOptions>& InOptions) {  }
 
+	/**
+	 * Returns the current encoder options if set or null otherwise.
+	 */
+	virtual const FTmvMediaEncoderOptions* GetEncoderOptions() const { return nullptr; }
+	
 	/**
 	 * Queries the encoder for the memory layout and format for frames of the given size, and doing so for each mips.
 	 * @param InParentJob Pointer to the parent job that contains the overall job settings.
diff --git a/Engine/Plugins/Media/TmvMedia/Source/TmvMedia/Public/Transcoder/TmvMediaFrameProducer.h b/Engine/Plugins/Media/TmvMedia/Source/TmvMedia/Public/Transcoder/TmvMediaFrameProducer.h
--- a/Engine/Plugins/Media/TmvMedia/Source/TmvMedia/Public/Transcoder/TmvMediaFrameProducer.h
+++ b/Engine/Plugins/Media/TmvMedia/Source/TmvMedia/Public/Transcoder/TmvMediaFrameProducer.h
@@ -2,6 +2,7 @@
 
 #pragma once
 
+#include "Misc/FrameRate.h"
 #include "Transcoder/TmvMediaTranscodeStage.h"
 #include "TmvMediaFrameProducer.generated.h"
 
@@ -9,6 +10,23 @@
 
 struct FTmvMediaFrameMips;
 
+/**
+ * Track information for the currently playing tracks from the frame producer stage.
+ */
+USTRUCT()
+struct FTmvMediaFrameProducerTrackInfo
+{
+	GENERATED_BODY()
+
+	/** Track's frame rate. */
+	UPROPERTY()
+	FFrameRate FrameRate = FFrameRate(24, 1);
+
+	/** Track's duration. */
+	UPROPERTY()
+	FTimespan Duration = FTimespan::Zero();
+};
+
 /**
  * Base class for producer stage of transcode pipeline.
  * @see UTmvMediaTranscodeJob
@@ -19,6 +37,19 @@ class UTmvMediaFrameProducer : public UTmvMediaTranscodeStage
 	GENERATED_BODY()
 public:
 	// todo: Add Media Info Request API for UI purposes.
+
+	/** Returns a copy of the current video track info. */
+	UE_API FTmvMediaFrameProducerTrackInfo GetVideoTrackInfo() const;
+
+	/** Set the video track info. */
+	UE_API void SetVideoTrackInfo(const FTmvMediaFrameProducerTrackInfo& InVideoTrackInfo);
+
+protected:
+	/** Track Info critical section. */
+	mutable FCriticalSection TrackInfoLock;
+
+	/** Information about the currently playing video track. */
+	FTmvMediaFrameProducerTrackInfo VideoTrackInfo;
 };
 
 #undef UE_API
\ No newline at end of file
diff --git a/Engine/Plugins/Media/TmvMedia/Source/TmvMedia/Public/Transcoder/TmvMediaTranscodeJob.h b/Engine/Plugins/Media/TmvMedia/Source/TmvMedia/Public/Transcoder/TmvMediaTranscodeJob.h
--- a/Engine/Plugins/Media/TmvMedia/Source/TmvMedia/Public/Transcoder/TmvMediaTranscodeJob.h
+++ b/Engine/Plugins/Media/TmvMedia/Source/TmvMedia/Public/Transcoder/TmvMediaTranscodeJob.h
@@ -59,6 +59,12 @@ struct FTmvMediaTranscodeJobSettings
 	/** Returns true if the input path is set. */
 	UE_API bool IsInputPathSet() const;
 
+	/** 
+	 * Returns the input frame rate setting. 
+	 * @remark will return the frame rate from the media source if specified.
+	 */
+	UE_API FFrameRate GetInputFramerate() const;
+
 	/** Returns true if the output path is set. */
 	UE_API bool IsOutputPathSet() const;
 
@@ -81,7 +87,7 @@ struct FTmvMediaTranscodeJobSettings
 	 */
 	UPROPERTY(EditAnywhere, Category = "Input Sequence", meta = (EditCondition = "InputSource ==  ETmvMediaTranscodeInputSource::File", EditConditionHides))
 	FFrameRate FrameRate = FFrameRate(24, 1);
-	
+
 	/**
 	 * Specify an existing media source as the input.
 	 */
@@ -100,9 +106,16 @@ struct FTmvMediaTranscodeJobSettings
 	UPROPERTY(EditAnywhere, Category = "Output Sequence", AdvancedDisplay, meta=(ClampMin=0, UIMin=0))
 	int32 ZeroPadFrameNumbers = 5;
 
-	/** todo - A media source will be made for the sequence. */
-	//UPROPERTY(EditAnywhere, Category = "Output Sequence")
-	//bool bMakeMediaSource = false;
+	/** If enabled, a media source asset will be made for the output sequence. */
+	UPROPERTY(EditAnywhere, Category = "Output Sequence")
+	bool bMakeOutputAsset = false;
+
+	/**
+	 * Destination content directory for the MediaSource asset to create or update (e.g. /Game/MediaSources) for the output sequence.
+	 * The asset name itself will be the current job name.
+	 */
+	UPROPERTY(EditAnywhere, Category = "Output Sequence", meta = (EditCondition = "bMakeOutputAsset", EditConditionHides, ContentDir))
+	FDirectoryPath OutputAssetDirectory;
 
 	/** todo - Enable specification of the transcoding range. */
 	//UPROPERTY(EditAnywhere, Category = "Output Sequence")
@@ -197,7 +210,7 @@ struct FTmvMediaTranscodingJobStats
  * There can only be one stage of each class per pipeline and the order is fixed as listed above.
  *
  * Todo - Audio and other streams:
- * The current implementation only deals a single video stream.
+ * The current implementation only deals with a single video stream.
  * For audio streams, we will need to make the pipeline more complex and deal with different streams into a common muxer stage.
  */
 UCLASS(MinimalAPI)
@@ -364,6 +377,9 @@ private:
 	void OnStopped(double InCurrentTime);
 	
 public:
+	/** Original Job Item Name. */
+	FString JobName;
+	
 	/** Common job settings. */
 	UPROPERTY()
 	FTmvMediaTranscodeJobSettings Settings;
@@ -371,7 +387,7 @@ public:
 	/** Stages in this pipeline. */
 	UPROPERTY()
 	TArray<TObjectPtr<UTmvMediaTranscodeStage>> Stages;	// Might use a map lookup with class
-
+	
 	/**
 	 * Overall job state.
 	 * It is driven by control functions like Start, Stop and then synchronized with the
diff --git a/Engine/Plugins/Media/TmvMedia/Source/TmvMediaEditor/Private/Customizations/TmvMediaDirectoryPathCustomization.cpp b/Engine/Plugins/Media/TmvMedia/Source/TmvMediaEditor/Private/Customizations/TmvMediaDirectoryPathCustomization.cpp
--- a/Engine/Plugins/Media/TmvMedia/Source/TmvMediaEditor/Private/Customizations/TmvMediaDirectoryPathCustomization.cpp
+++ b/Engine/Plugins/Media/TmvMedia/Source/TmvMediaEditor/Private/Customizations/TmvMediaDirectoryPathCustomization.cpp
@@ -2,6 +2,7 @@
 
 #include "TmvMediaDirectoryPathCustomization.h"
 
+#include "ContentBrowserDataDragDropOp.h"
 #include "DetailWidgetRow.h"
 #include "IDetailChildrenBuilder.h"
 #include "Input/DragAndDrop.h"
@@ -30,7 +31,12 @@ void FTmvMediaDirectoryPathCustomization::CustomizeChildren(TSharedRef<IProperty
 	IDetailChildrenBuilder& InChildBuilder, IPropertyTypeCustomizationUtils& InCustomizationUtils)
 {
 	PathStringProperty = InPropertyHandle->GetChildHandle("Path");
-	
+
+	// Support same modifiers as FDirectoryPathStructCustomization::CustomizeHeader for content browser.
+	bContentDir = InPropertyHandle->HasMetaData( TEXT("ContentDir") ) || InPropertyHandle->HasMetaData(TEXT("LongPackageName"));
+	const bool bForceShowPluginContent = InPropertyHandle->HasMetaData(TEXT("ForceShowPluginContent"));
+	const bool bForceShowEngineContent = InPropertyHandle->HasMetaData(TEXT("ForceShowEngineContent"));
+
 	// Add Directory picker.
 	InChildBuilder.AddCustomRow(LOCTEXT("DirectoryPathLabel", "Path"))
 		.NameContent()
@@ -58,6 +64,10 @@ void FTmvMediaDirectoryPathCustomization::CustomizeChildren(TSharedRef<IProperty
 				.BrowseTitle(LOCTEXT("BrowseTitle", "Directory picker..."))
 				.DirectoryPath(this, &FTmvMediaDirectoryPathCustomization::GetDirectoryPath, /*bReturnEmptyOnMultiSelect*/ false)
 				.OnDirectoryChanged(this, &FTmvMediaDirectoryPathCustomization::HandleDirectoryPickerDirectoryChanged)
+				.IsEnabled(this, &FTmvMediaDirectoryPathCustomization::IsEnabled)
+				.ContentDir(bContentDir)
+				.ForceShowPluginContent(bForceShowPluginContent)
+				.ForceShowEngineContent(bForceShowEngineContent)
 			]
 		];
 }
@@ -75,16 +85,30 @@ FString FTmvMediaDirectoryPathCustomization::GetDirectoryPath(bool bReturnEmptyO
 	return FilePath;
 }
 
-void FTmvMediaDirectoryPathCustomization::HandleDirectoryPickerDirectoryChanged(const FString& InPickedDirectory)
+void FTmvMediaDirectoryPathCustomization::HandleDirectoryPickerDirectoryChanged(const FString& InPickedDirectory, bool bInContentDir)
 {
 	if (PathStringProperty)
 	{
-		PathStringProperty->SetValue(UE::TmvMedia::PathUtils::GetSanitizedPath(InPickedDirectory));
+		using namespace UE::TmvMedia::PathUtils;
+		PathStringProperty->SetValue(bInContentDir ? InPickedDirectory : GetSanitizedPath(InPickedDirectory));
 	}
 }
 
 bool FTmvMediaDirectoryPathCustomization::HandleVerifyDrag(TSharedPtr<FDragDropOperation> InDragOperation)
 {
+	if (bContentDir)
+	{
+		if (InDragOperation && InDragOperation->IsOfType<FContentBrowserDataDragDropOp>())
+		{
+			TSharedPtr<FContentBrowserDataDragDropOp> ContentDragDropOp = StaticCastSharedPtr<FContentBrowserDataDragDropOp>(InDragOperation);
+			return ContentDragDropOp &&
+				(!ContentDragDropOp->GetDraggedFolders().IsEmpty() ||
+				 !ContentDragDropOp->GetDraggedFiles().IsEmpty() ||
+				 !ContentDragDropOp->GetDraggedItems().IsEmpty());
+		}
+		return false;
+	}
+
 	if (InDragOperation && InDragOperation->IsOfType<FExternalDragOperation>())
 	{
 		TSharedPtr<FExternalDragOperation> ExternalDragOperation = StaticCastSharedPtr<FExternalDragOperation>(InDragOperation);
@@ -96,16 +120,57 @@ bool FTmvMediaDirectoryPathCustomization::HandleVerifyDrag(TSharedPtr<FDragDropO
 FReply FTmvMediaDirectoryPathCustomization::HandleDropEvent(const FGeometry& InGeometry, const FDragDropEvent& InDropEvent)
 {
 	const TSharedPtr<FDragDropOperation> DropOperation = InDropEvent.GetOperation();
+
+	if (bContentDir)
+	{
+		if (DropOperation && DropOperation->IsOfType<FContentBrowserDataDragDropOp>())
+		{
+			TSharedPtr<FContentBrowserDataDragDropOp> ContentDragDropOp = StaticCastSharedPtr<FContentBrowserDataDragDropOp>(DropOperation);
+			if (ContentDragDropOp)
+			{
+				FString ContentPath;
+				if (!ContentDragDropOp->GetDraggedFolders().IsEmpty())
+				{
+					ContentPath = ContentDragDropOp->GetDraggedFolders()[0].GetInternalPath().ToString();
+				}
+				else if (!ContentDragDropOp->GetDraggedFiles().IsEmpty())
+				{
+					ContentPath = FPaths::GetPath(ContentDragDropOp->GetDraggedFiles()[0].GetInternalPath().ToString());
+				}
+				else if (!ContentDragDropOp->GetDraggedItems().IsEmpty())
+				{
+					const FContentBrowserItem& Item = ContentDragDropOp->GetDraggedItems()[0];
+					ContentPath = Item.IsFolder()
+						? Item.GetInternalPath().ToString()
+						: FPaths::GetPath(Item.GetInternalPath().ToString());
+				}
+
+				if (!ContentPath.IsEmpty())
+				{
+					HandleDirectoryPickerDirectoryChanged(ContentPath, /*bInContentDir*/true);
+					return FReply::Handled();
+				}
+			}
+		}
+		return FReply::Unhandled();
+	}
+
 	if (DropOperation && DropOperation->IsOfType<FExternalDragOperation>())
 	{
 		TSharedPtr<FExternalDragOperation> ExternalDragOperation = StaticCastSharedPtr<FExternalDragOperation>(DropOperation);
 		if (ExternalDragOperation->HasFiles())
 		{
 			const FString SanitizedPath = UE::TmvMedia::PathUtils::GetSanitizedPath(ExternalDragOperation->GetFiles()[0]);
-			HandleDirectoryPickerDirectoryChanged(UE::TmvMedia::PathUtils::GetDirectoryFullPath(SanitizedPath));
+			HandleDirectoryPickerDirectoryChanged(UE::TmvMedia::PathUtils::GetDirectoryFullPath(SanitizedPath), /*bInContentDir*/false);
 			return FReply::Handled();
 		}
 	}
 	return FReply::Unhandled();
 }
+
+bool FTmvMediaDirectoryPathCustomization::IsEnabled() const
+{
+	return PathStringProperty ? PathStringProperty->IsEditable() : false;
+}
+
 #undef LOCTEXT_NAMESPACE
diff --git a/Engine/Plugins/Media/TmvMedia/Source/TmvMediaEditor/Private/Customizations/TmvMediaDirectoryPathCustomization.h b/Engine/Plugins/Media/TmvMedia/Source/TmvMediaEditor/Private/Customizations/TmvMediaDirectoryPathCustomization.h
--- a/Engine/Plugins/Media/TmvMedia/Source/TmvMediaEditor/Private/Customizations/TmvMediaDirectoryPathCustomization.h
+++ b/Engine/Plugins/Media/TmvMedia/Source/TmvMediaEditor/Private/Customizations/TmvMediaDirectoryPathCustomization.h
@@ -37,7 +37,7 @@ private:
 	/**
 	 * Called when a directory is selected.
 	 */
-	void HandleDirectoryPickerDirectoryChanged(const FString& InPickedDirectory);
+	void HandleDirectoryPickerDirectoryChanged(const FString& InPickedDirectory, bool bInContentDir);
 
 	/**
 	 * Called when a drag operation enters the widget. 
@@ -49,7 +49,13 @@ private:
 	 */
 	FReply HandleDropEvent(const FGeometry& InGeometry, const FDragDropEvent& InDropEvent);
 
+	/** Returns true if the widget is enabled (editable). */
+	bool IsEnabled() const;
+
 	/** Pointer to the string that will be set when changing the path */
 	TSharedPtr<IPropertyHandle> PathStringProperty;
+
+	/** Indicates that the path will be picked using the Slate-style directory picker inside the game Content dir. */
+	bool bContentDir = false;
 };
 
diff --git a/Engine/Plugins/Media/TmvMedia/Source/TmvMediaEditor/Private/Widgets/STmvMediaDirectoryPicker.cpp b/Engine/Plugins/Media/TmvMedia/Source/TmvMediaEditor/Private/Widgets/STmvMediaDirectoryPicker.cpp
--- a/Engine/Plugins/Media/TmvMedia/Source/TmvMediaEditor/Private/Widgets/STmvMediaDirectoryPicker.cpp
+++ b/Engine/Plugins/Media/TmvMedia/Source/TmvMediaEditor/Private/Widgets/STmvMediaDirectoryPicker.cpp
@@ -2,7 +2,9 @@
 
 #include "STmvMediaDirectoryPicker.h"
 
+#include "ContentBrowserModule.h"
 #include "DesktopPlatformModule.h"
+#include "IContentBrowserSingleton.h"
 #include "Framework/Application/SlateApplication.h"
 #include "Widgets/Images/SImage.h"
 #include "Widgets/Input/SButton.h"
@@ -16,6 +18,9 @@ void STmvMediaDirectoryPicker::Construct( const FArguments& InArgs )
 	BrowseDirectory = InArgs._BrowseDirectory;
 	DirectoryPath = InArgs._DirectoryPath;
 	OnDirectoryChanged = InArgs._OnDirectoryChanged;
+	bContentDir = InArgs._ContentDir;
+	bForceShowPluginContent = InArgs._ForceShowPluginContent;
+	bForceShowEngineContent = InArgs._ForceShowEngineContent;
 	
 	TSharedPtr<SButton> OpenButton;
 	ChildSlot
@@ -60,7 +65,7 @@ void STmvMediaDirectoryPicker::Construct( const FArguments& InArgs )
 
 void STmvMediaDirectoryPicker::OnDirectoryTextCommited(const FText& InText, ETextCommit::Type InCommitType)
 {
-	OnDirectoryChanged.ExecuteIfBound(InText.ToString());
+	OnDirectoryChanged.ExecuteIfBound(InText.ToString(), bContentDir);
 }
 
 FText STmvMediaDirectoryPicker::GetDirectoryTextBoxText() const
@@ -71,11 +76,17 @@ FText STmvMediaDirectoryPicker::GetDirectoryTextBoxText() const
 FReply STmvMediaDirectoryPicker::OnBrowseButtonClicked()
 {
 	const FString DefaultPath = BrowseDirectory.IsSet() ? BrowseDirectory.Get() : DirectoryPath.Get();
-	FString OutDirectory;
-
-	if ( OpenPlatformDirectoryPicker(OutDirectory, DefaultPath) )
+	if (bContentDir)
+	{
+		OpenContentPathPicker(DefaultPath);
+	}
+	else
 	{
-		OnDirectoryChanged.ExecuteIfBound(OutDirectory);
+		FString OutDirectory;
+		if ( OpenPlatformDirectoryPicker(OutDirectory, DefaultPath) )
+		{
+			OnDirectoryChanged.ExecuteIfBound(OutDirectory, /*bInContentDir*/ false);
+		}
 	}
 
 	return FReply::Handled();
@@ -108,4 +119,46 @@ bool STmvMediaDirectoryPicker::OpenPlatformDirectoryPicker(FString& OutDirectory
 	return bFolderSelected;
 }
 
+void STmvMediaDirectoryPicker::OpenContentPathPicker(const FString& InDefaultPath)
+{
+	FContentBrowserModule& ContentBrowserModule = FModuleManager::LoadModuleChecked<FContentBrowserModule>("ContentBrowser");
+	FPathPickerConfig PathPickerConfig;
+	// PathPickerConfig.DefaultPath = InDefaultPath; // That seems to make it not work.
+	PathPickerConfig.bAllowContextMenu = false;
+	PathPickerConfig.OnPathSelected = FOnPathSelected::CreateSP(this, &STmvMediaDirectoryPicker::OnContentPathPicked);
+	PathPickerConfig.bForceShowEngineContent = bForceShowEngineContent;
+	PathPickerConfig.bForceShowPluginContent = bForceShowPluginContent;
+
+	FMenuBuilder MenuBuilder(true, nullptr);
+	MenuBuilder.AddWidget(SNew(SBox)
+		.WidthOverride(300.0f)
+		.HeightOverride(300.0f)
+		[
+			ContentBrowserModule.Get().CreatePathPicker(PathPickerConfig)
+		], FText());
+
+	if (ContentPickerMenu.IsValid())
+	{
+		ContentPickerMenu->Dismiss();
+	}
+	
+	ContentPickerMenu = FSlateApplication::Get().PushMenu(
+		AsShared(),
+		FWidgetPath(),
+		MenuBuilder.MakeWidget(),
+		FSlateApplication::Get().GetCursorPos(),
+		FPopupTransitionEffect(FPopupTransitionEffect::ContextMenu));
+}
+
+void STmvMediaDirectoryPicker::OnContentPathPicked(const FString& InPickedPath)
+{
+	if (ContentPickerMenu.IsValid())
+	{
+		ContentPickerMenu->Dismiss();
+		ContentPickerMenu.Reset();
+	}
+
+	OnDirectoryChanged.ExecuteIfBound(InPickedPath, /*bInContentDir*/ true);
+}
+
 #undef LOCTEXT_NAMESPACE
diff --git a/Engine/Plugins/Media/TmvMedia/Source/TmvMediaEditor/Private/Widgets/STmvMediaDirectoryPicker.h b/Engine/Plugins/Media/TmvMedia/Source/TmvMediaEditor/Private/Widgets/STmvMediaDirectoryPicker.h
--- a/Engine/Plugins/Media/TmvMedia/Source/TmvMediaEditor/Private/Widgets/STmvMediaDirectoryPicker.h
+++ b/Engine/Plugins/Media/TmvMedia/Source/TmvMediaEditor/Private/Widgets/STmvMediaDirectoryPicker.h
@@ -10,6 +10,7 @@
 #include "Widgets/DeclarativeSyntaxSupport.h"
 #include "Widgets/SCompoundWidget.h"
 
+class IMenu;
 class SEditableTextBox;
 
 /**
@@ -23,12 +24,15 @@ class STmvMediaDirectoryPicker : public SCompoundWidget
 {
 public:
 	/** Declares a delegate that is executed when a directory was picked. */
-	DECLARE_DELEGATE_OneParam(FOnDirectoryChanged, const FString& /*Directory*/);
+	DECLARE_DELEGATE_TwoParams(FOnDirectoryChanged, const FString& /*Directory*/, bool /*bInContentDir*/);
 
 	SLATE_BEGIN_ARGS(STmvMediaDirectoryPicker)
 		: _BrowseButtonToolTip(NSLOCTEXT("STmvMediaDirectoryPicker", "BrowseButtonToolTip", "Choose a directory from this computer"))
 		, _Font()
 		, _IsReadOnly(false)
+		, _ContentDir(false)
+		, _ForceShowPluginContent(false)
+		, _ForceShowEngineContent(false)
 	{ }
 
 		/** Browse button image resource. */
@@ -55,6 +59,15 @@ public:
 		/** Whether the path text box can be modified by the user. */
 		SLATE_ATTRIBUTE(bool, IsReadOnly)
 
+		/** Indicates that the path will be picked using the Slate-style directory picker inside the game Content dir. */
+		SLATE_ARGUMENT(bool, ContentDir)
+
+		/** Indicates that the asset pickers should always show plugin content */
+		SLATE_ARGUMENT(bool, ForceShowPluginContent)
+
+		/** Indicates that the asset pickers should always show engine content */
+		SLATE_ARGUMENT(bool, ForceShowEngineContent)
+
 		/** Called when a path has been picked or modified. */
 		SLATE_EVENT(FOnDirectoryChanged, OnDirectoryChanged)
 
@@ -77,6 +90,12 @@ private:
 	/** Utility function to open a directory picker dialog. */
 	bool OpenPlatformDirectoryPicker(FString& OutDirectory, const FString& InDefaultPath) const;
 
+	/** Utility function to open a content path picker drop down menu. */
+	void OpenContentPathPicker(const FString& InDefaultPath);
+
+	/** Content Picker Menu - Called when a content path is picked. */
+	void OnContentPathPicked(const FString& InPickedPath);
+
 private:
 	/** Holds the directory path to browse by default. */
 	TAttribute<FString> BrowseDirectory;
@@ -92,4 +111,16 @@ private:
 
 	/** Editable text box for the directory path. */
 	TSharedPtr<SEditableTextBox> DirectoryTextBox;
+
+	/** The content pick button popup menu*/
+	TSharedPtr<IMenu> ContentPickerMenu;
+
+	/** Indicates that the path will be picked using the Slate-style directory picker inside the game Content dir. */
+	bool bContentDir = false;
+
+	/** Indicates that the asset pickers should always show plugin content */
+	bool bForceShowPluginContent = false;
+
+	/** Indicates that the asset pickers should always show engine content */
+	bool bForceShowEngineContent = false;
 };
\ No newline at end of file
diff --git a/Engine/Plugins/Media/TmvMedia/Source/TmvMediaEditor/Private/Widgets/STmvMediaTranscoderJobControls.cpp b/Engine/Plugins/Media/TmvMedia/Source/TmvMediaEditor/Private/Widgets/STmvMediaTranscoderJobControls.cpp
--- a/Engine/Plugins/Media/TmvMedia/Source/TmvMediaEditor/Private/Widgets/STmvMediaTranscoderJobControls.cpp
+++ b/Engine/Plugins/Media/TmvMedia/Source/TmvMediaEditor/Private/Widgets/STmvMediaTranscoderJobControls.cpp
@@ -11,6 +11,7 @@
 #include "TmvMediaEditorLog.h"
 #include "TmvMediaTranscodeListHandle.h"
 #include "TmvMediaTranscodeNotification.h"
+#include "TmvMediaTranscodeTaskMakeMediaSource.h"
 #include "Transcoder/TmvMediaFileSequenceFrameProducer.h"
 #include "Transcoder/TmvMediaFrameConverter.h"
 #include "Transcoder/TmvMediaPlayerFrameProducer.h"
@@ -48,8 +49,8 @@ void STmvMediaTranscoderJobControls::Construct(const FArguments& InArgs, const T
 					.IsEnabled(this, &STmvMediaTranscoderJobControls::CanStartJobs)
 					.OnClicked(this, &STmvMediaTranscoderJobControls::OnStartJobs)
 					.ForegroundColor(FSlateColor::UseStyle())
-					.Text(LOCTEXT("StartJobButtonLabel", "Start Job(s)"))
-					.ToolTipText(LOCTEXT("StartJobButtonToolTip", "Start transcoding selected jobs."))
+					.Text(LOCTEXT("StartJobButtonLabel", "Start Selected Job(s)"))
+					.ToolTipText(this, &STmvMediaTranscoderJobControls::GetStartJobsToolTip)
 			]
 
 		+ SHorizontalBox::Slot()
@@ -77,14 +78,15 @@ namespace UE::TmvMediaEditor::Transcode
 	{
 		CleanupJob(InJob, FApp::GetCurrentTime());
 	}
-	
+
 	// todo: move this in a transcode job builder in the runtime module.
 	UTmvMediaTranscodeJob* MakeJob(const FTmvMediaTranscodeListItem* InJobItem)
 	{
 		UTmvMediaTranscodeJob* TranscodeJob = NewObject<UTmvMediaTranscodeJob>(GetTransientPackage(), NAME_None);
 		TranscodeJob->SetId(InJobItem->Id);
+		TranscodeJob->JobName = InJobItem->Name;
 		TranscodeJob->Settings = InJobItem->Settings;
-		
+
 		if (TranscodeJob->Settings.bUseMediaPlayer)
 		{
 			TranscodeJob->AddStage(NewObject<UTmvMediaPlayerFrameProducer>(TranscodeJob, NAME_None));
@@ -93,13 +95,16 @@ namespace UE::TmvMediaEditor::Transcode
 		{
 			TranscodeJob->AddStage(NewObject<UTmvMediaFileSequenceFrameProducer>(TranscodeJob, NAME_None));
 		}
-		
+
 		TranscodeJob->AddStage(NewObject<UTmvMediaFrameConverter>(TranscodeJob, NAME_None));
 		TranscodeJob->AddStage(NewObject<UTmvMediaTmvFrameEncoder>(TranscodeJob, NAME_None));
 		//TranscodeJob->AddStage(NewObject<UTmvMediaTestFrameEncoder>(TranscodeJob, NAME_None));
 		TranscodeJob->AddStage(NewObject<UTmvMediaTranscodeMuxer>(TranscodeJob, NAME_None));
 		TranscodeJob->GetStage<UTmvMediaFrameEncoder>()->SetEncoderOptions(InJobItem->EncoderOptions);
-		
+
+		// Extension editor tasks (can only be implemented in editor module).
+		UE::TmvMediaEditor::TranscodeTask::AddMakeOrUpdateMediaSourceTask(TranscodeJob);
+
 		return TranscodeJob;
 	}
 }
@@ -154,9 +159,77 @@ bool STmvMediaTranscoderJobControls::IsJobProcessingOrPending(const FGuid& InJob
 	return false;
 }
 
+FText STmvMediaTranscoderJobControls::GetStartJobsToolTip() const
+{
+	if (CanStartJobs())
+	{
+		return LOCTEXT("StartJobButtonToolTip", "Start transcoding selected jobs.");
+	}
+	
+	// If the jobs can't be started, find why and drop a hint to the user in the tool tip.
+	UTmvMediaTranscodeList* List = ListHandle ? ListHandle->Get() : nullptr;
+	if (!List)
+	{
+		return LOCTEXT("StartJobButtonToolTip_InvalidListHandle", "Error: Invalid List Handle.");
+	}
+
+	FTextBuilder Reasons;
+	Reasons.Indent();	// Indent our list of reasons.
+
+	const TArray<int32> CurrentSelection = ListHandle->GetCurrentSelection();
+	
+	// Check that we don't have any items already started or that is not configured to run.
+	for (const int32 SelectedItem : CurrentSelection)
+	{
+		if (List->IsValidItemIndex(SelectedItem))
+		{
+			const FTmvMediaTranscodeListItem& JobItem = List->GetItem(SelectedItem); 
+			if (!CanStartJobItem(JobItem))
+			{
+				if (IsJobProcessingOrPending(JobItem.Id))
+				{
+					Reasons.AppendLineFormat(LOCTEXT("JobFailReason_IsProcessing", "-Item \"{0}\": Is already processing or pending."), FText::FromString(JobItem.Name));
+				}
+				else
+				{
+					if (!JobItem.Settings.IsInputPathSet())
+					{
+						Reasons.AppendLineFormat(LOCTEXT("JobFailReason_NoInputPath", "-Item \"{0}\": Input Path Not Set."), FText::FromString(JobItem.Name));
+					}
+					if (!JobItem.Settings.IsOutputPathSet())
+					{
+						Reasons.AppendLineFormat(LOCTEXT("JobFailReason_NoOutputPath", "-Item \"{0}\": Output Path Not Set."), FText::FromString(JobItem.Name));
+					}
+					if (JobItem.Settings.bMakeOutputAsset && JobItem.Settings.OutputAssetDirectory.Path.IsEmpty())
+					{
+						Reasons.AppendLineFormat(LOCTEXT("JobFailReason_NoOutputAssetDirectory", "-Item \"{0}\": Output Asset Directory Not Set."), FText::FromString(JobItem.Name));
+					}
+				}
+				
+				// Protection for too many items.
+				if (Reasons.GetNumLines() > 20)
+				{
+					Reasons.AppendLine(LOCTEXT("JobFailReason_Ellipsis", "..."));
+					break;
+				}
+			}
+		}
+	}
+	
+	if (CurrentSelection.IsEmpty() || Reasons.GetNumLines() == 0)
+	{
+		Reasons.AppendLine(LOCTEXT("StartJobButtonToolTip_NoSelectedJobs", "-Job selection is empty."));
+	}
+
+	return FText::Format(LOCTEXT("StartJobButtonToolTip_JobFailReasons", "Cannot start transcoding selected jobs:\n{0}"), Reasons.ToText());
+}
+
 bool STmvMediaTranscoderJobControls::CanStartJobItem(const FTmvMediaTranscodeListItem& InJobItem) const
 {
-	return InJobItem.Settings.IsInputPathSet() && InJobItem.Settings.IsOutputPathSet() && !IsJobProcessingOrPending(InJobItem.Id);
+	return InJobItem.Settings.IsInputPathSet()
+		&& InJobItem.Settings.IsOutputPathSet()
+		&& (!InJobItem.Settings.bMakeOutputAsset || !InJobItem.Settings.OutputAssetDirectory.Path.IsEmpty()) 
+		&& !IsJobProcessingOrPending(InJobItem.Id);
 }
 
 bool STmvMediaTranscoderJobControls::CanStartJobs() const
diff --git a/Engine/Plugins/Media/TmvMedia/Source/TmvMediaEditor/Private/Widgets/STmvMediaTranscoderJobControls.h b/Engine/Plugins/Media/TmvMedia/Source/TmvMediaEditor/Private/Widgets/STmvMediaTranscoderJobControls.h
--- a/Engine/Plugins/Media/TmvMedia/Source/TmvMediaEditor/Private/Widgets/STmvMediaTranscoderJobControls.h
+++ b/Engine/Plugins/Media/TmvMedia/Source/TmvMediaEditor/Private/Widgets/STmvMediaTranscoderJobControls.h
@@ -37,6 +37,9 @@ private:
 	/** Stops and discards the current job and discards all pending jobs. */
 	void CleanUp();
 
+	/** Returns the Start Jobs button tool tip. */
+	FText GetStartJobsToolTip() const;
+
 	/** Returns true if the given job item can be started. */
 	bool CanStartJobItem(const FTmvMediaTranscodeListItem& InJobItem) const;
 	
@@ -45,7 +48,7 @@ private:
 
 	/** Called when we click on the start jobs button. */
 	FReply OnStartJobs();
-
+	
 	/** Returns true if the current set of running or pending jobs can be cancelled. */
 	bool CanCancelJobs() const;
 
diff --git a/Engine/Plugins/Media/TmvMedia/Source/TmvMediaEditor/Private/Widgets/TmvMediaTranscodeTaskMakeMediaSource.cpp b/Engine/Plugins/Media/TmvMedia/Source/TmvMediaEditor/Private/Widgets/TmvMediaTranscodeTaskMakeMediaSource.cpp
new file mode 100644
--- /dev/null
+++ b/Engine/Plugins/Media/TmvMedia/Source/TmvMediaEditor/Private/Widgets/TmvMediaTranscodeTaskMakeMediaSource.cpp
@@ -0,0 +1,330 @@
+// Copyright Epic Games, Inc. All Rights Reserved.
+
+#include "TmvMediaTranscodeTaskMakeMediaSource.h"
+
+#include "AssetRegistry/AssetRegistryModule.h"
+#include "Encoder/TmvMediaEncoderOptions.h"
+#include "HAL/FileManager.h"
+#include "MediaSource.h"
+#include "Misc/PackageName.h"
+#include "Misc/Paths.h"
+#include "ObjectTools.h"
+#include "TmvMediaEditorLog.h"
+#include "Transcoder/TmvMediaFrameEncoder.h"
+#include "Transcoder/TmvMediaFrameProducer.h"
+#include "Transcoder/TmvMediaTranscodeJob.h"
+#include "UObject/Package.h"
+#include "UObject/SavePackage.h"
+#include "Utils/TmvMediaPathUtils.h"
+
+namespace UE::TmvMediaEditor::TranscodeTask
+{
+	/**
+	 * Creates a MediaSource asset at InFullAssetPath pointing to InMediaPath.
+	 *
+	 * Uses UMediaSource::SpawnMediaSourceForString to create the concrete media source type
+	 * (typically UImgMediaSource or UFileMediaSource) based on the file extension of the output files.
+	 */
+	UMediaSource* CreateMediaSourceAsset(const FString& InFullAssetPath, const FString& InMediaPath, const FString& InFileExtension, UPackage* InPackage)
+	{
+		FString ResolvedFilePath = InMediaPath;
+
+		if (FPaths::DirectoryExists(InMediaPath))
+		{
+			// Find any file in the output directory, filtered on the given extension if available, to pass to SpawnMediaSourceForString.
+			TArray<FString> FoundFiles;
+
+			IFileManager::Get().FindFiles(FoundFiles, *InMediaPath, !InFileExtension.IsEmpty() ? *InFileExtension : nullptr);
+
+			// Do a minimum of clean up to get a valid file.
+			FoundFiles.RemoveAll([](const FString& InFileName) { return InFileName.StartsWith(TEXT(".")); });
+			FoundFiles.Sort();
+
+			if (FoundFiles.IsEmpty())
+			{
+				UE_LOG(LogTmvMediaEditor, Warning,
+					TEXT("MakeOrUpdateMediaSourceAsset: No output files found in '%s'; skipping media source creation."),
+					*InMediaPath);
+				return nullptr;
+			}
+
+			ResolvedFilePath = InMediaPath / FoundFiles[0];
+		}
+
+		// Make sure the resolved file path exists
+		if (!FPaths::FileExists(ResolvedFilePath))
+		{
+			UE_LOG(LogTmvMediaEditor, Warning,
+				TEXT("MakeOrUpdateMediaSourceAsset: '%s' does not exist; skipping media source creation."),
+				*ResolvedFilePath);
+			return nullptr;
+		}
+
+		UMediaSource* MediaSource = UMediaSource::SpawnMediaSourceForString(ResolvedFilePath, InPackage);
+		if (!MediaSource)
+		{
+			UE_LOG(LogTmvMediaEditor, Warning,
+				TEXT("MakeOrUpdateMediaSourceAsset: SpawnMediaSourceForString returned null for '%s'."),
+				*ResolvedFilePath);
+			return nullptr;
+		}
+
+		const FString AssetName = FPackageName::GetShortName(InFullAssetPath);
+		MediaSource->Rename(*AssetName, InPackage);
+		MediaSource->SetFlags(RF_Public | RF_Standalone);
+		MediaSource->MarkPackageDirty();
+		FAssetRegistryModule::AssetCreated(MediaSource);
+
+		return MediaSource;
+	}
+
+	/** Possible results of the set property operations. */ 
+	enum class EPropertyResult : int8
+	{
+		/** Update operation failed: missing parameter, wrong type, etc. */
+		Failed = -1,
+		/** Update operation done successfully. */
+		Done = 0,
+		/** Value was already up to date and nothing needed to be done. */
+		UpToDate = 1
+	};
+
+	/**
+	 * Use the blueprint function from the known media source classes to set the media path.
+	 * Supporting UImgMediaSource "SetSequencePath" and UFileMediaSource "SetFilePath". 
+	 * @param InMediaSource Media Source to configure 
+	 * @param InPath Desired path to be set as the media source's media file or sequence path.
+	 * @return Either "Failed" if the known functions are not found or have the wrong signature, or "Done" otherwise.
+	 */
+	EPropertyResult SetMediaSourceMediaPath(UMediaSource* InMediaSource, const FString& InPath)
+	{
+		// Try to find UImgMediaSource "SetSequencePath".
+		UFunction* Func = InMediaSource->FindFunction(TEXT("SetSequencePath"));
+
+		// Try to find UFileMediaSource "SetFilePath", which as the same signature.
+		if (!Func)
+		{
+			Func = InMediaSource->FindFunction(TEXT("SetFilePath"));
+		}
+
+		if (!Func)
+		{
+			return EPropertyResult::Failed;
+		}
+
+		// Allocate a buffer matching the function's actual param layout so we don't
+		// hard-code assumptions about the UHT-generated params struct.
+		TArray<uint8> ParamBuffer;
+		ParamBuffer.AddZeroed(Func->ParmsSize);
+
+		// Initialize all params and locate the Path parameter by name + type.
+		bool bFoundPath = false;
+		for (TFieldIterator<FProperty> It(Func); It && It->HasAnyPropertyFlags(CPF_Parm); ++It)
+		{
+			It->InitializeValue_InContainer(ParamBuffer.GetData());
+
+			if (!bFoundPath && It->GetFName() == TEXT("Path"))
+			{
+				if (FStrProperty* StrProp = CastField<FStrProperty>(*It))
+				{
+					StrProp->SetPropertyValue_InContainer(ParamBuffer.GetData(), InPath);
+					bFoundPath = true;
+				}
+			}
+		}
+
+		if (!bFoundPath)
+		{
+			UE_LOG(LogTmvMediaEditor, Warning,
+				TEXT("SetMediaSourceSequencePath: \"%s\" signature has changed; \"Path\" param not found on \"%s\"."),
+				*Func->GetName(), *InMediaSource->GetClass()->GetName());
+		}
+		else
+		{
+			InMediaSource->ProcessEvent(Func, ParamBuffer.GetData());
+			InMediaSource->MarkPackageDirty();
+		}
+
+		// Destroy param values — important for FString to release heap memory.
+		for (TFieldIterator<FProperty> It(Func); It && It->HasAnyPropertyFlags(CPF_Parm); ++It)
+		{
+			It->DestroyValue_InContainer(ParamBuffer.GetData());
+		}
+		
+		return bFoundPath ? EPropertyResult::Done : EPropertyResult::Failed;
+	}
+
+	/**
+	 * Utility function to set the "FrameRateOverride" property of the media source.
+	 * 
+	 * @param InMediaSource Media Source to configure 
+	 * @param InFrameRate Desired frame rate to set in the media source.
+	 * @return	Either "Done" if the frame rate was set, "UpToDate" if the frame rate was already the desired value, 
+	 *			or "Failed" if the property was not found.
+	 */
+	EPropertyResult SetMediaSourceFrameRate(UMediaSource* InMediaSource, const FFrameRate& InFrameRate)
+	{
+		// Update the framerate in the img media source.
+		if (FStructProperty* const FrameRateProperty = FindFProperty<FStructProperty>(InMediaSource->GetClass(), TEXT("FrameRateOverride")))
+		{
+			if (void* PropertyValuePtr = FrameRateProperty->ContainerPtrToValuePtr<FFrameRate>(InMediaSource, 0))
+			{
+				// Read the value
+				FFrameRate PropertyValue;
+				FrameRateProperty->CopyCompleteValue(&PropertyValue, PropertyValuePtr);
+				
+				if (PropertyValue != InFrameRate)
+				{
+					FrameRateProperty->CopyCompleteValue(PropertyValuePtr, &InFrameRate);
+
+					// notifies editor/content browser
+					FPropertyChangedEvent ChangeEvent(FrameRateProperty);
+					InMediaSource->PostEditChangeProperty(ChangeEvent);
+
+					InMediaSource->MarkPackageDirty();
+					return EPropertyResult::Done;
+				}
+				return EPropertyResult::UpToDate;
+			}
+		}
+		return EPropertyResult::Failed;
+	}
+
+	/** Save the given media source to file. */
+	bool SaveMediaSourceAsset(UPackage* InPackage, UMediaSource* InMediaSource, const FString& InFullAssetPath)
+	{
+		const FString PackageFilename = FPackageName::LongPackageNameToFilename(InFullAssetPath, FPackageName::GetAssetPackageExtension());
+		FSavePackageArgs SaveArgs;
+		SaveArgs.TopLevelFlags = RF_Public | RF_Standalone;
+		return UPackage::SavePackage(InPackage, InMediaSource, *PackageFilename, SaveArgs);
+	}
+
+	const FString& GetValidString(const FString& InString, const FString& InDefault)
+	{
+		return InString.IsEmpty() ? InDefault : InString;
+	}
+
+	/** Make a safe asset name from the job settings. */
+	FString MakeMediaSourceFullAssetName(const UTmvMediaTranscodeJob* InTranscodeJob)
+	{
+		// Using the job name as the media source name. If it is empty, we use a default name.
+		const FString MediaSourceName = GetValidString(InTranscodeJob->JobName, TEXT("TranscodeOutput"));
+		const FString MediaSourcePath = GetValidString(InTranscodeJob->Settings.OutputAssetDirectory.Path, TEXT("/Game"));
+		return MediaSourcePath / MediaSourceName;
+	}
+
+	/**
+	 * Creates or replaces a MediaSource asset at InAssetPath pointing to InSequencePath.
+	 */
+	void MakeOrUpdateMediaSourceAsset(UTmvMediaTranscodeJob* InTranscodeJob)
+	{
+		// Only perform the task if the job succeeded without errors and was not canceled.
+		const bool bJobSucceeded = !InTranscodeJob->bIsError && InTranscodeJob->StopReason == ETmvMediaTranscodeJobStopReason::Completed;
+
+		if (bJobSucceeded && InTranscodeJob->Settings.bMakeOutputAsset && !InTranscodeJob->Settings.OutputAssetDirectory.Path.IsEmpty())
+		{
+			// This is the fallback setting, but should be overriden by the media file's frame rate.
+			FFrameRate FrameRate = InTranscodeJob->Settings.GetInputFramerate();
+
+			// Fetch the frame rate from the frame producer, which is going to be what the media file was.
+			if (UTmvMediaFrameProducer* FrameProducer = InTranscodeJob->GetStage<UTmvMediaFrameProducer>())
+			{
+				const FTmvMediaFrameProducerTrackInfo VideoTrackInfo = FrameProducer->GetVideoTrackInfo();
+				FrameRate = VideoTrackInfo.FrameRate;
+			}
+
+			// Using the job name as the media source name.
+			const FString FullMediaSourcePath = MakeMediaSourceFullAssetName(InTranscodeJob);
+			const FString MediaSourceName = FPackageName::GetShortName(FullMediaSourcePath);
+			const FString FullMediaPath = InTranscodeJob->Settings.GetAbsoluteOutputPath();
+
+			UPackage* MediaSourcePackage = CreatePackage(*FullMediaSourcePath);
+
+			// See if we have an existing media source.
+			UMediaSource* MediaSource = Cast<UMediaSource>(FindObject<UObject>(MediaSourcePackage, *MediaSourceName));
+
+			// Accumulator for existing media source(s) that might be discarded if it can't be updated.
+			TArray<UObject*> DiscardedMediaSourcesToConsolidate;
+			
+			bool bNeedSaveMediaSource = false;
+
+			// Existing media source
+			if (MediaSource)
+			{
+				using namespace UE::TmvMedia::PathUtils;
+				// Check if the sequence path is already correct.
+				FString ExistingMediaPath = GetMediaSourceMediaFullPath(MediaSource);
+				
+				// Try to set the path
+				if (!FPaths::IsSamePath(ExistingMediaPath, FullMediaPath))
+				{
+					bNeedSaveMediaSource = true; // At least the path will be changed.
+
+					// Try to update the sequence path.
+					const EPropertyResult SetPathResult = SetMediaSourceMediaPath(MediaSource, FullMediaPath);
+
+					// If we couldn't update it, we will have to discard this object and create a new one.
+					if (SetPathResult == EPropertyResult::Failed || !FPaths::IsSamePath(GetMediaSourceMediaFullPath(MediaSource), FullMediaPath))
+					{
+						// Fallback method:
+						// We will discard the current media source and create a new one with the correct path and class.
+						// However, we will need to consolidate in memory objects to update references.
+						MediaSource->Rename(nullptr, GetTransientPackage(), REN_DontCreateRedirectors | REN_NonTransactional);
+						DiscardedMediaSourcesToConsolidate.Add(MediaSource);
+						MediaSource = nullptr;
+					}
+				}
+			}
+
+			// Create new media source.
+			if (!MediaSource)
+			{
+				// Retrieve the produced file extension, if any.
+				FString FileExtension;
+				if (const UTmvMediaFrameEncoder* EncoderStage = InTranscodeJob->GetStage<UTmvMediaFrameEncoder>())
+				{
+					if (const FTmvMediaEncoderOptions* EncoderOptions = EncoderStage->GetEncoderOptions())
+					{
+						FileExtension = EncoderOptions->GetFileSequenceExtension();
+					}
+				}
+
+				MediaSource = CreateMediaSourceAsset(FullMediaSourcePath, FullMediaPath, FileExtension, MediaSourcePackage);
+				if (!MediaSource)
+				{
+					return; // CreateMediaSourceAsset already logged an error message.
+				}
+				bNeedSaveMediaSource = true;
+
+				// Consolidate in memory objects to update references, etc.
+				if (!DiscardedMediaSourcesToConsolidate.IsEmpty())
+				{
+					UE_LOG(LogTmvMediaEditor, Log, TEXT("Consolidating references to new media source \"%s\" ... "), *FullMediaSourcePath);
+					ObjectTools::ConsolidateObjects(MediaSource, DiscardedMediaSourcesToConsolidate, /*bShowDeleteConfirmation*/ false);
+				}
+			}
+
+			if (MediaSource)
+			{
+				if (SetMediaSourceFrameRate(MediaSource, FrameRate) == EPropertyResult::Done)
+				{
+					bNeedSaveMediaSource = true; // Frame rate was modified, will need saving.
+				}
+
+				// Only save the asset if it was created of modified to avoid invoking source control if not needed.
+				if (bNeedSaveMediaSource)
+				{
+					if (!SaveMediaSourceAsset(MediaSourcePackage, MediaSource, FullMediaSourcePath))
+					{
+						UE_LOG(LogTmvMediaEditor, Error, TEXT("Failed to save media source \"%s\""), *FullMediaSourcePath);
+					}
+				}
+			}
+		}
+	}
+
+	void AddMakeOrUpdateMediaSourceTask(UTmvMediaTranscodeJob* InTranscodeJob)
+	{
+		InTranscodeJob->GetOnTranscodeJobFinished().AddStatic(&MakeOrUpdateMediaSourceAsset);
+	}
+}
diff --git a/Engine/Plugins/Media/TmvMedia/Source/TmvMediaEditor/Private/Widgets/TmvMediaTranscodeTaskMakeMediaSource.h b/Engine/Plugins/Media/TmvMedia/Source/TmvMediaEditor/Private/Widgets/TmvMediaTranscodeTaskMakeMediaSource.h
new file mode 100644
--- /dev/null
+++ b/Engine/Plugins/Media/TmvMedia/Source/TmvMediaEditor/Private/Widgets/TmvMediaTranscodeTaskMakeMediaSource.h
@@ -0,0 +1,16 @@
+// Copyright Epic Games, Inc. All Rights Reserved.
+
+#pragma once
+
+#include "HAL/Platform.h"
+
+class UTmvMediaTranscodeJob;
+
+namespace UE::TmvMediaEditor::TranscodeTask
+{
+	/**
+	 * Adds a task that will create/update a MediaSource asset for the output sequence when the job is finished.
+	 * @param InTranscodeJob Transcode job to add the task to.
+	 */
+	void AddMakeOrUpdateMediaSourceTask(UTmvMediaTranscodeJob* InTranscodeJob);
+}
diff --git a/Engine/Plugins/Media/TmvMedia/Source/TmvMediaEditor/TmvMediaEditor.Build.cs b/Engine/Plugins/Media/TmvMedia/Source/TmvMediaEditor/TmvMediaEditor.Build.cs
--- a/Engine/Plugins/Media/TmvMedia/Source/TmvMediaEditor/TmvMediaEditor.Build.cs
+++ b/Engine/Plugins/Media/TmvMedia/Source/TmvMediaEditor/TmvMediaEditor.Build.cs
@@ -9,16 +9,19 @@ namespace UnrealBuildTool.Rules
 			PrivateDependencyModuleNames.AddRange(
 				new string[] {
 					"ApplicationCore",
+					"AssetRegistry",
+					"ContentBrowserData",
 					"Core",
 					"CoreUObject",
 					"DesktopWidgets",
 					"EditorWidgets",
 					"InputCore",
+					"MediaAssets",
 					"Slate",
 					"SlateCore",
 					"TmvMedia",
 					"UnrealEd",
-					"WorkspaceMenuStructure", 
+					"WorkspaceMenuStructure",
 				});
 		}
 	}
