[TmvMediaEditor] Make/Update Media Source Task: Fix fatal error about saving the package while not being fully loaded.

#Synced-CL 52257653

diff --git a/Engine/Plugins/Media/TmvMedia/Source/TmvMediaEditor/Private/Widgets/TmvMediaTranscodeTaskMakeMediaSource.cpp b/Engine/Plugins/Media/TmvMedia/Source/TmvMediaEditor/Private/Widgets/TmvMediaTranscodeTaskMakeMediaSource.cpp
--- a/Engine/Plugins/Media/TmvMedia/Source/TmvMediaEditor/Private/Widgets/TmvMediaTranscodeTaskMakeMediaSource.cpp
+++ b/Engine/Plugins/Media/TmvMedia/Source/TmvMediaEditor/Private/Widgets/TmvMediaTranscodeTaskMakeMediaSource.cpp
@@ -240,6 +240,15 @@ namespace UE::TmvMediaEditor::TranscodeTask
 
 			UPackage* MediaSourcePackage = CreatePackage(*FullMediaSourcePath);
 
+			if (!MediaSourcePackage)
+			{
+				UE_LOG(LogTmvMediaEditor, Error, TEXT("Unable to create media source \"%s\" ... "), *FullMediaSourcePath);
+				return;
+			}
+			
+			// We may have to save this package, so it has to be fully loaded.
+			MediaSourcePackage->FullyLoad();
+
 			// See if we have an existing media source.
 			UMediaSource* MediaSource = Cast<UMediaSource>(FindObject<UObject>(MediaSourcePackage, *MediaSourceName));
 
