[TmvMedia] Apv Encoder: Don't issue a warning if the maximum of cores is higher than the max supported by the OAPV implementation.

The number of threads is maxed at 32 in the OAPV implementation. Selecting a number of threads of 0 in Encoder Settings would then use the number of cores on the system and issue a warning in the logs if the CPU has more than 32 cores. We had this with a system of 96 cores (Threadripper 9995). The warning will only be issued if the number of threads was explicit.

#Synced-CL 52257653

diff --git a/Engine/Plugins/Media/TmvMedia/Source/ApvMedia/Private/Apv/ApvEncoder.cpp b/Engine/Plugins/Media/TmvMedia/Source/ApvMedia/Private/Apv/ApvEncoder.cpp
--- a/Engine/Plugins/Media/TmvMedia/Source/ApvMedia/Private/Apv/ApvEncoder.cpp
+++ b/Engine/Plugins/Media/TmvMedia/Source/ApvMedia/Private/Apv/ApvEncoder.cpp
@@ -57,10 +57,10 @@ namespace UE::ApvMedia
 
 	int32 GetValidatedNumThreads(int32 InNumThreads)
 	{
-		int32 NumThreads = InNumThreads > 0 ? InNumThreads : FPlatformMisc::NumberOfCores();
+		int32 NumThreads = InNumThreads > 0 ? InNumThreads : FMath::Min(FPlatformMisc::NumberOfCores(), OAPV_MAX_THREADS);
 		if (NumThreads > OAPV_MAX_THREADS)
 		{
-			UE_LOG(LogApvMedia, Warning, TEXT("Requested number of threads (%d) exceeds maximum OpenAPV threads (%d)."), NumThreads, OAPV_MAX_THREADS);
+			UE_LOG(LogApvMedia, Warning, TEXT("Explicitly requested number of threads (%d) exceeds maximum OpenAPV threads (%d)."), NumThreads, OAPV_MAX_THREADS);
 		}
 		return FMath::Clamp(NumThreads, 1, OAPV_MAX_THREADS);
 	}
