[TmvMedia] Tmv Transcoder: Media Player Frame Producer edge case fixes

Fix for some edge cases:
- Using new GetTrackDuration to help with media that have tracks of different durations. Protron might give the shortest track, but not Electra, so trying to transcode with Electra when an audio track is longuer than the video track can lead to the frame producer getting stuck. Using the new api resolve this problem since we are specifically requesting the duration of the video track we are transcoding.
- Adding a tolerance to the upper bound of the media to compensate for rounding errors depending on players and containers. Despite using the new api, some players can round the frame duration down one tick or round the track duration up one tick, leading to inconsistent or incorrect detection of transcoding completion.
- Adding time outs while waiting for media player to produce frames. If the frame producer gets stuck again, at least the time out will get it unstuck. If it is stuck on the last frame, the transcode job can complete successfully (even though it timed out on the last frame).

#Synced-CL 52257653

diff --git a/Engine/Plugins/Media/TmvMedia/Source/TmvMedia/Internal/Transcoder/TmvMediaPlayerFrameProducer.h b/Engine/Plugins/Media/TmvMedia/Source/TmvMedia/Internal/Transcoder/TmvMediaPlayerFrameProducer.h
--- a/Engine/Plugins/Media/TmvMedia/Source/TmvMedia/Internal/Transcoder/TmvMediaPlayerFrameProducer.h
+++ b/Engine/Plugins/Media/TmvMedia/Source/TmvMedia/Internal/Transcoder/TmvMediaPlayerFrameProducer.h
@@ -41,7 +41,7 @@ public:
 
 protected:
 	/** Sub-Tick function that handles the "playing" state. */
-	void HandlePlaying(UTmvMediaTranscodeJob* InParentJob);
+	void HandlePlaying(UTmvMediaTranscodeJob* InParentJob, float InDeltaTime);
 
 	/** CLose the player and associated resources. */
 	void Close(const UTmvMediaTranscodeJob* InParentJob);
@@ -118,6 +118,9 @@ protected:
 	
 	/** The current time we are processing. */
 	FTimespan CurrentTime;
+	
+	/** Keep track of the previous tick's rendered sample time to detect when new frames are rendered. */
+	FTimespan PreviousRenderedSampleTime;
 
 	/** Maximum number of concurrently processing frames. */
 	int32 MaxProcessingFrames = 1;
@@ -134,6 +137,7 @@ protected:
 		Closed,
 		Opening,
 		Open,
+		Preparing,
 		Playing,
 		MediaCompleted,
 		NotSupported,
diff --git a/Engine/Plugins/Media/TmvMedia/Source/TmvMedia/Private/Transcoder/TmvMediaPlayerFrameProducer.cpp b/Engine/Plugins/Media/TmvMedia/Source/TmvMedia/Private/Transcoder/TmvMediaPlayerFrameProducer.cpp
--- a/Engine/Plugins/Media/TmvMedia/Source/TmvMedia/Private/Transcoder/TmvMediaPlayerFrameProducer.cpp
+++ b/Engine/Plugins/Media/TmvMedia/Source/TmvMedia/Private/Transcoder/TmvMediaPlayerFrameProducer.cpp
@@ -135,6 +135,7 @@ bool UTmvMediaPlayerFrameProducer::Start(UTmvMediaTranscodeJob* InParentJob)
 
 	// Start playing.
 	MediaStartTime = CurrentTime = MediaFrameDuration = FTimespan::Zero();
+	PreviousRenderedSampleTime = FTimespan::MinValue();
 
 	FMediaPlayerOptions Options;
 	Options.SetAllAsOptional();
@@ -283,7 +284,14 @@ void UTmvMediaPlayerFrameProducer::Tick(UTmvMediaTranscodeJob* InParentJob, cons
 					MediaPlayer->SetBlockOnTimeRange(TRange<FTimespan>(CurrentTime, CurrentTime + FrameDuration));
 				
 					MediaPlayer->Play();
-					CurrentState = EState::Playing;
+					if (MediaPlayer->IsPreparing())
+					{
+						CurrentState = EState::Preparing;
+					}
+					else
+					{
+						CurrentState = EState::Playing;
+					}
 					WatchdogTimeRemaining = 5.0f;
 				}
 			}
@@ -294,7 +302,7 @@ void UTmvMediaPlayerFrameProducer::Tick(UTmvMediaTranscodeJob* InParentJob, cons
 			}
 			break;
 		}
-	case EState::Playing:
+	case EState::Preparing:
 		{
 			if (MediaPlayer->IsPreparing())
 			{
@@ -307,19 +315,25 @@ void UTmvMediaPlayerFrameProducer::Tick(UTmvMediaTranscodeJob* InParentJob, cons
 			}
 			else
 			{
-				HandlePlaying(InParentJob);
+				CurrentState = EState::Playing;
 				WatchdogTimeRemaining = 5.0f;	// Reset watchdog.
+				HandlePlaying(InParentJob, DeltaTime);
 			}
 			break;
 		}
+	case EState::Playing:
+		{
+			HandlePlaying(InParentJob, DeltaTime);
+			break;
+		}
 	case EState::MediaCompleted:
-		
+
 		if ((WatchdogTimeRemaining -= DeltaTime) < 0.0f)
 		{
 			UE_LOG(LogTmvMedia, Error, TEXT("Transcode Frame Producer: Timed out waiting for media to finish encoding."));
 			CurrentState = EState::TimedOut;
 		}
-		
+
 		if (ProducerResource)
 		{
 			const int32 CurrentSubmittedCount = ProducerResource->GetSubmittedCount();
@@ -344,7 +358,7 @@ void UTmvMediaPlayerFrameProducer::Tick(UTmvMediaTranscodeJob* InParentJob, cons
 					}
 				}
 
-				WatchdogTimeRemaining = 5.0f;	// Reset watchdog.
+				WatchdogTimeRemaining = 10.0f; // Reset watchdog. Waiting for encoding (slow).
 			}
 
 			// Wait for final captures to process.
@@ -372,7 +386,7 @@ void UTmvMediaPlayerFrameProducer::Tick(UTmvMediaTranscodeJob* InParentJob, cons
 	}
 }
 
-void UTmvMediaPlayerFrameProducer::HandlePlaying(UTmvMediaTranscodeJob* InParentJob)
+void UTmvMediaPlayerFrameProducer::HandlePlaying(UTmvMediaTranscodeJob* InParentJob, float InDeltaTime)
 {
 	using namespace UE::TmvMedia;
 	
@@ -386,11 +400,26 @@ void UTmvMediaPlayerFrameProducer::HandlePlaying(UTmvMediaTranscodeJob* InParent
 	// Get which frame the player is on.
 	const FTimespan RenderedSampleTime = ProducerResource->GetLastRenderedSampleTime();
 
+	if ((WatchdogTimeRemaining -= InDeltaTime) < 0.0f)
+	{
+		UE_LOG(LogTmvMedia, Error, TEXT("Transcode Frame Producer: Timed out waiting for media to render the next frame."));
+		CurrentState = EState::TimedOut;
+		return;
+	}
+
 	if (RenderedSampleTime == FTimespan::MinValue())
 	{
 		return;	// not rendered yet.
 	}
 
+	// Detect when new frames are rendered.
+	if (RenderedSampleTime != PreviousRenderedSampleTime)
+	{
+		WatchdogTimeRemaining = 5.0f;	// Reset watchdog for next frame.
+	}
+
+	PreviousRenderedSampleTime = RenderedSampleTime;
+
 	if (MediaFrameDuration == FTimespan::Zero())
 	{
 		// Try to update one last time.
@@ -422,7 +451,7 @@ void UTmvMediaPlayerFrameProducer::HandlePlaying(UTmvMediaTranscodeJob* InParent
 	const FTimespan EndRenderedSampleTime = RenderedSampleTime + FrameDuration;
 
 	// Render time caught up with play time?
-	if ( EndRenderedSampleTime >= CurrentTime)
+	if (EndRenderedSampleTime >= CurrentTime)
 	{
 		int32 RenderedFrameIndex = 0;
 		int32 TotalFrames = 0;
@@ -453,9 +482,13 @@ void UTmvMediaPlayerFrameProducer::HandlePlaying(UTmvMediaTranscodeJob* InParent
 
 			if (!MediaPlaybackRange.IsEmpty())
 			{
-				if (CurrentTime < MediaPlaybackRange.GetUpperBoundValue())
+				// For detecting the last frame, it is possible there are small imprecision on the upper bound
+				// of the range or the frame's duration. Those values may be rounded down or up by 1 tick depending on players or
+				// containers. We will give a tolerance of a small fraction of a frame to the detection of the last frame to compensate.
+				const FTimespan UpperBoundTolerance = FMath::Max(FrameDuration * 0.1, FTimespan(2));
+				if (CurrentTime < MediaPlaybackRange.GetUpperBoundValue() - UpperBoundTolerance)
 				{
-					TRange<FTimespan> NextBlockOnRange = TRange<FTimespan>(CurrentTime, CurrentTime + FrameDuration);
+					TRange<FTimespan> NextBlockOnRange = TRange<FTimespan>(CurrentTime + FTimespan(1), CurrentTime + FrameDuration);
 					MediaPlayer->SetBlockOnTimeRange(NextBlockOnRange);
 				}
 				else
@@ -474,6 +507,10 @@ void UTmvMediaPlayerFrameProducer::HandlePlaying(UTmvMediaTranscodeJob* InParent
 		}
 		else
 		{
+			// Since the encoder can be very slow, we want to reset the producer's watchdog while
+			// we wait, otherwise this could lead to timeouts because of the encoder.
+			WatchdogTimeRemaining = 5.0f;
+
 			UE_LOG(LogTmvMedia, Verbose, TEXT("Transcode Frame Producer: WAITING  Player Time: %s Rendered Time: %s Playback Range: [%s, %s] Frame: %d"),
 				*ToString(CurrentTime),
 				*ToString(RenderedSampleTime),
@@ -529,11 +566,29 @@ void UTmvMediaPlayerFrameProducer::TryUpdatePlaybackRange()
 	const FTimespan MediaDuration = MediaPlayer->GetDuration();
 	if (MediaDuration > FTimespan::Zero() && MediaDuration < FTimespan::MaxValue())
 	{
-		MediaPlaybackRange = MediaPlayer->GetPlaybackTimeRange(EMediaTimeRangeType::Absolute);
+		const int32 SelectedTrack = MediaPlayer->GetSelectedTrack(EMediaPlayerTrack::Video);
+		const TOptional<FTimespan> TrackDuration = MediaPlayer->GetTrackDuration(EMediaPlayerTrack::Video, SelectedTrack);
+
+		// Priority to track duration, more precise in case other tracks in present.
+		if (TrackDuration.IsSet())
+		{
+			MediaPlaybackRange = TRange<FTimespan>(FTimespan::Zero(), TrackDuration.GetValue());
+		}
+		else
+		{
+			// Fallback to GetPlaybackTimeRange api.
+			MediaPlaybackRange = MediaPlayer->GetPlaybackTimeRange(EMediaTimeRangeType::Absolute);
+			if (MediaPlaybackRange.IsEmpty())
+			{
+				// Fallback to GetDuration api (older api).
+				MediaPlaybackRange = TRange<FTimespan>(FTimespan::Zero(), MediaDuration);
+			}
+		}
 
 		UE_LOG(LogTmvMedia, Verbose, TEXT("Transcode Frame Producer: Media Playback Range [%s, %s]."),
 			*MediaPlaybackRange.GetLowerBoundValue().ToString(), *MediaPlaybackRange.GetUpperBoundValue().ToString());
-					
+
+		// Make sure the media has a video track.
 		if (MediaPlayer->GetNumTracks(EMediaPlayerTrack::Video) <= 0 || MediaPlayer->GetTrackFormat(EMediaPlayerTrack::Video, 0) < 0)
 		{
 			UE_LOG(LogTmvMedia, Error, TEXT("Transcode Frame Producer: Media has no video track."));
