[TmvMedia] Transcoder Frame Producer: Reducing memory consumption for the media player frame producer by adjusting the pre-caching window to the smallest.

We don't need to precache for the transcoder because the producing pipeline is usually faster than the consuming (encoder) pipeline. And we want to use a little memory as possible to help encoding very large 16K videos.

#Synced-CL 52257653

diff --git a/Engine/Plugins/Media/TmvMedia/Source/TmvMedia/Private/Transcoder/TmvMediaPlayerFrameProducer.cpp b/Engine/Plugins/Media/TmvMedia/Source/TmvMedia/Private/Transcoder/TmvMediaPlayerFrameProducer.cpp
--- a/Engine/Plugins/Media/TmvMedia/Source/TmvMedia/Private/Transcoder/TmvMediaPlayerFrameProducer.cpp
+++ b/Engine/Plugins/Media/TmvMedia/Source/TmvMedia/Private/Transcoder/TmvMediaPlayerFrameProducer.cpp
@@ -140,8 +140,9 @@ bool UTmvMediaPlayerFrameProducer::Start(UTmvMediaTranscodeJob* InParentJob)
 	FMediaPlayerOptions Options;
 	Options.SetAllAsOptional();
 
-	// Experimenting...
-	Options.InternalCustomOptions.Emplace(MediaPlayerOptionValues::Environment(), MediaPlayerOptionValues::Environment_Sequencer());
+	// Environment_Sequencer vs Environment_Preview:
+	// in Environment_Preview, the player will not try to pre-cache as it is not necessary for transcoding.
+	Options.InternalCustomOptions.Emplace(MediaPlayerOptionValues::Environment(), MediaPlayerOptionValues::Environment_Preview());
 	
 	Options.PlayOnOpen = EMediaPlayerOptionBooleanOverride::Disabled;
 	Options.Loop = EMediaPlayerOptionBooleanOverride::Disabled;
@@ -151,11 +152,12 @@ bool UTmvMediaPlayerFrameProducer::Start(UTmvMediaTranscodeJob* InParentJob)
 	Options.Tracks.Audio = -1;
 	// For image media, we avoid filling the global cache which will needlessly hold onto frame data.
 	Options.InternalCustomOptions.Emplace(MediaPlayerOptionValues::ImgMediaSmartCacheEnabled(), FVariant(true));
-	Options.InternalCustomOptions.Emplace(MediaPlayerOptionValues::ImgMediaSmartCacheTimeToLookAhead(), FVariant(0.2f));
+	// Setting a look ahead time to 0 means the minimum of precaching, which should be 2 frames.
+	Options.InternalCustomOptions.Emplace(MediaPlayerOptionValues::ImgMediaSmartCacheTimeToLookAhead(), FVariant(0.0f));
 	Options.InternalCustomOptions.Emplace(MediaPlayerOptionValues::ViewMediaTexture(), FVariant(FSoftObjectPath(MediaTexture).ToString()));
 
 	CurrentState = EState::Opening;
-	WatchdogTimeRemaining = 10.0f;	// We will wait for 10 seconds for the video to open.
+	WatchdogTimeRemaining = 10.0f;	// We will wait up to 10 seconds for the video to open.
 	
 	if (!MediaPlayer->OpenSourceWithOptions(MediaSource, Options))
 	{
